PlainProps move code to experimental engine plugin

#rnx

diff --git a/Engine/Plugins/Experimental/PlainProps/Source/Public/PlainPropsLoad.h b/Engine/Plugins/Experimental/PlainProps/Source/Public/PlainPropsLoad.h
new file mode 100644
--- /dev/null
+++ b/Engine/Plugins/Experimental/PlainProps/Source/Public/PlainPropsLoad.h
@@ -0,0 +1,41 @@
+// Copyright Epic Games, Inc. All Rights Reserved.
+
+#pragma once
+
+#include "Containers/ArrayView.h"
+#include "PlainPropsRead.h"
+#include "Memory/MemoryFwd.h"
+
+namespace PlainProps 
+{
+	
+class FCustomBindings;
+class FDeclarations;
+struct FLoadBatch;
+class FRangeBinding;
+class FSchemaBindings;
+struct FTypedRange;
+
+
+PLAINPROPS_API FLoadBatch* CreateLoadPlans(FReadBatchId ReadId, const FDeclarations& Declarations, const FCustomBindings& Customs, const FSchemaBindings& Schemas, TConstArrayView<FStructSchemaId> RuntimeIds);
+PLAINPROPS_API void DestroyLoadPlans(FLoadBatch* Batch);
+PLAINPROPS_API void LoadStruct(uint8* Dst, FByteReader Src, FStructSchemaId Id, const FLoadBatch& Batch);
+PLAINPROPS_API void LoadStruct(uint8* Dst, FStructView Src, const FLoadBatch& Batch);
+PLAINPROPS_API void ConstructAndLoadStruct(uint8* Dst, FByteReader Src, FStructSchemaId Id, const FLoadBatch& Batch);
+PLAINPROPS_API void ConstructAndLoadStruct(uint8* Dst, FStructView Src, const FLoadBatch& Batch);
+PLAINPROPS_API void LoadRange(uint8* Dst, FRangeView Src, ERangeSizeType MaxType, TConstArrayView<FRangeBinding> Bindings, const FLoadBatch& Batch);
+
+//using FItemConstructor = void (*)(FLoadRangeContext& Ctx);
+
+//struct FLoadRangeBinding
+//{
+//	FItemConstructor	MakeItems;
+//	const void*			UserContext;
+//};
+
+// Possible optimization - opt-in leaf range load method to avoid copying and repeated calls when target representation is not a range
+//template<typename T> using TDirectLeafLoader = void (*)(TRangeView<T> Items);
+//using FDirectBoolLoader = void (*)(FBoolRangeView Bits);
+
+} // namespace PlainProps
+
diff --git a/Engine/Plugins/Experimental/PlainProps/Source/Public/PlainPropsWrite.h b/Engine/Plugins/Experimental/PlainProps/Source/Public/PlainPropsWrite.h
new file mode 100644
--- /dev/null
+++ b/Engine/Plugins/Experimental/PlainProps/Source/Public/PlainPropsWrite.h
@@ -0,0 +1,39 @@
+// Copyright Epic Games, Inc. All Rights Reserved.
+
+#pragma once
+
+#include "PlainPropsTypes.h"
+#include "Containers/Array.h"
+#include "Templates/UniquePtr.h"
+
+namespace PlainProps
+{
+struct FBuiltSchemas;
+struct FBuiltStruct;
+class FDebugIds;
+class FIdIndexerBase;
+class FNestedScopeIndexer;
+class FParametricTypeIndexer;
+struct FWriteIds;
+
+class FWriter
+{
+public:
+	PLAINPROPS_API FWriter(const FIdIndexerBase& DeclaredIds, const FBuiltSchemas& InSchemas);
+	PLAINPROPS_API ~FWriter();
+	
+	PLAINPROPS_API bool								Uses(FNameId BuiltId) const;
+	PLAINPROPS_API TConstArrayView<FNestedScope>	GetUsedScopes() const;
+	PLAINPROPS_API TConstArrayView<FParametricType>	GetUsedParametrics() const;
+	PLAINPROPS_API FOptionalStructSchemaId			GetWriteId(FStructSchemaId BuiltId) const;
+
+	PLAINPROPS_API void								WriteSchemas(TArray64<uint8>& Out) const;
+	PLAINPROPS_API void								WriteMembers(TArray64<uint8>& Out, FStructSchemaId BuiltId, const FBuiltStruct& Struct) const;
+
+private:
+	const FBuiltSchemas&							Schemas;
+	const FDebugIds&								Debug;
+	TUniquePtr<FWriteIds>							NewIds;
+};
+
+} // namespace PlainProps
\ No newline at end of file
diff --git a/Engine/Source/Programs/PlainPropsTests/PlainPropsTests.Build.cs b/Engine/Source/Programs/PlainPropsTests/PlainPropsTests.Build.cs
--- a/Engine/Source/Programs/PlainPropsTests/PlainPropsTests.Build.cs
+++ b/Engine/Source/Programs/PlainPropsTests/PlainPropsTests.Build.cs
@@ -29,8 +29,7 @@ public class PlainPropsTests : TestModuleRules
 	{
 		PrivateDependencyModuleNames.AddRange(
 			new string[] {
-				"Core",
-				"CoreUObject"
+				"PlainProps"
 			});
 
 		if (Target.bBuildWithEditorOnlyData)
diff --git a/Engine/Source/Programs/PlainPropsTests/PlainPropsTests.Target.cs b/Engine/Source/Programs/PlainPropsTests/PlainPropsTests.Target.cs
--- a/Engine/Source/Programs/PlainPropsTests/PlainPropsTests.Target.cs
+++ b/Engine/Source/Programs/PlainPropsTests/PlainPropsTests.Target.cs
@@ -10,7 +10,7 @@ public class PlainPropsTestsTarget : TestTargetRules
 		// Collects all tests decorated with #if WITH_LOW_LEVELTESTS from dependencies
 		bWithLowLevelTestsOverride = true;
 		bCompileWithPluginSupport = true;
-		bCompileAgainstCoreUObject = true;
+		bCompileAgainstCoreUObject = false;
 
 		bBuildWithEditorOnlyData = Target.Platform.IsInGroup(UnrealPlatformGroup.Desktop)
 			&& (Target.Configuration == UnrealTargetConfiguration.Debug || Target.Configuration == UnrealTargetConfiguration.Development);
