Speculative .uplugin parse fix

#rnx

diff --git a/Engine/Plugins/Experimental/PlainProps/PlainProps.uplugin b/Engine/Plugins/Experimental/PlainProps/PlainProps.uplugin
--- a/Engine/Plugins/Experimental/PlainProps/PlainProps.uplugin
+++ b/Engine/Plugins/Experimental/PlainProps/PlainProps.uplugin
@@ -7,12 +7,10 @@
 	"CreatedByURL" : "https://epicgames.com",
 	"EnabledByDefault" : false,
 	"IsExperimentalVersion": true,
-	
-	"Modules" :
-	[
+	"Modules": [
 		{
 			"Name" : "PlainProps",
-			"Type" : "Runtime",
+			"Type" : "Runtime"
 		}
 	]
-}
+}
\ No newline at end of file
