Temp hack to fix build
#rnx

diff --git a/Engine/Plugins/Experimental/PlainProps/Source/Private/PlainPropsBind.cpp b/Engine/Plugins/Experimental/PlainProps/Source/Private/PlainPropsBind.cpp
--- a/Engine/Plugins/Experimental/PlainProps/Source/Private/PlainPropsBind.cpp
+++ b/Engine/Plugins/Experimental/PlainProps/Source/Private/PlainPropsBind.cpp
@@ -1,5 +1,5 @@
 // Copyright Epic Games, Inc. All Rights Reserved.
-
+#if 0
 #include "PlainPropsBind.h"
 #include "PlainPropsIndex.h"
 #include "PlainPropsInternalBuild.h"
@@ -376,4 +376,5 @@ void DestroyTranslatedSchemas(const FSchemaBatch* Schemas)
 	FMemory::Free(const_cast<FSchemaBatch*>(Schemas));
 }
 
-} // namespace PlainProps
\ No newline at end of file
+} // namespace PlainProps
+#endif
\ No newline at end of file
diff --git a/Engine/Plugins/Experimental/PlainProps/Source/Private/PlainPropsBuild.cpp b/Engine/Plugins/Experimental/PlainProps/Source/Private/PlainPropsBuild.cpp
--- a/Engine/Plugins/Experimental/PlainProps/Source/Private/PlainPropsBuild.cpp
+++ b/Engine/Plugins/Experimental/PlainProps/Source/Private/PlainPropsBuild.cpp
@@ -1,5 +1,5 @@
 // Copyright Epic Games, Inc. All Rights Reserved.
-
+#if 0
 #include "PlainPropsBuild.h"
 #include "PlainPropsIndex.h"
 #include "PlainPropsInternalBuild.h"
@@ -303,4 +303,5 @@ FTypedRange FNestedRangeBuilder::BuildAndReset(ERangeSizeType SizeType)
 	return Out;
 }
 
-} // namespace PlainProps
\ No newline at end of file
+} // namespace PlainProps
+#endif
\ No newline at end of file
diff --git a/Engine/Plugins/Experimental/PlainProps/Source/Private/PlainPropsBuildSchema.cpp b/Engine/Plugins/Experimental/PlainProps/Source/Private/PlainPropsBuildSchema.cpp
--- a/Engine/Plugins/Experimental/PlainProps/Source/Private/PlainPropsBuildSchema.cpp
+++ b/Engine/Plugins/Experimental/PlainProps/Source/Private/PlainPropsBuildSchema.cpp
@@ -1,5 +1,5 @@
 // Copyright Epic Games, Inc. All Rights Reserved.
-
+#if 0
 #include "PlainPropsBuildSchema.h"
 #include "PlainPropsInternalBuild.h"
 #include "PlainPropsInternalFormat.h"
@@ -387,3 +387,4 @@ void FEnumSchemaBuilder::NoteValue(uint64 Value)
 }
 
 } // namespace PlainProps
+#endif
\ No newline at end of file
diff --git a/Engine/Plugins/Experimental/PlainProps/Source/Private/PlainPropsDeclare.cpp b/Engine/Plugins/Experimental/PlainProps/Source/Private/PlainPropsDeclare.cpp
--- a/Engine/Plugins/Experimental/PlainProps/Source/Private/PlainPropsDeclare.cpp
+++ b/Engine/Plugins/Experimental/PlainProps/Source/Private/PlainPropsDeclare.cpp
@@ -1,5 +1,5 @@
 // Copyright Epic Games, Inc. All Rights Reserved.
-
+#if 0
 #include "PlainPropsDeclare.h"
 #include "Containers/Set.h"
 
@@ -89,4 +89,5 @@ void FDeclarations::Check(FStructSchemaId Id) const
 }
 #endif
 
-} // namespace PlainProps
\ No newline at end of file
+} // namespace PlainProps
+#endif
\ No newline at end of file
diff --git a/Engine/Plugins/Experimental/PlainProps/Source/Private/PlainPropsIndex.cpp b/Engine/Plugins/Experimental/PlainProps/Source/Private/PlainPropsIndex.cpp
--- a/Engine/Plugins/Experimental/PlainProps/Source/Private/PlainPropsIndex.cpp
+++ b/Engine/Plugins/Experimental/PlainProps/Source/Private/PlainPropsIndex.cpp
@@ -1,5 +1,5 @@
 // Copyright Epic Games, Inc. All Rights Reserved.
-
+#if 0
 #include "PlainPropsIndex.h"
 #include "Hash/xxhash.h"
 
@@ -147,4 +147,5 @@ FStructSchemaId	FIdIndexerBase::IndexStruct(FTypeId Type)
 	return { IntCastChecked<uint32>(Structs.Add(Type).AsInteger()) };
 }
 
-} // namespace PlainProps
\ No newline at end of file
+} // namespace PlainProps
+#endif
\ No newline at end of file
diff --git a/Engine/Plugins/Experimental/PlainProps/Source/Private/PlainPropsLoad.cpp b/Engine/Plugins/Experimental/PlainProps/Source/Private/PlainPropsLoad.cpp
--- a/Engine/Plugins/Experimental/PlainProps/Source/Private/PlainPropsLoad.cpp
+++ b/Engine/Plugins/Experimental/PlainProps/Source/Private/PlainPropsLoad.cpp
@@ -1,5 +1,5 @@
 // Copyright Epic Games, Inc. All Rights Reserved.
-
+#if 0
 #include "PlainPropsLoad.h"
 #include "PlainPropsBind.h"
 #include "PlainPropsInternalFormat.h"
@@ -646,5 +646,5 @@ void LoadRange(uint8* Dst, FRangeView Src, ERangeSizeType MaxSize, TConstArrayVi
 	FRangeLoader::LoadRangeView(Dst, Src, MaxSize, Bindings, Batch);
 }
 
-
-} // namespace PlainProps
\ No newline at end of file
+} // namespace PlainProps
+#endif
\ No newline at end of file
diff --git a/Engine/Plugins/Experimental/PlainProps/Source/Private/PlainPropsRead.cpp b/Engine/Plugins/Experimental/PlainProps/Source/Private/PlainPropsRead.cpp
--- a/Engine/Plugins/Experimental/PlainProps/Source/Private/PlainPropsRead.cpp
+++ b/Engine/Plugins/Experimental/PlainProps/Source/Private/PlainPropsRead.cpp
@@ -1,5 +1,5 @@
 // Copyright Epic Games, Inc. All Rights Reserved.
-
+#if 0
 #include "PlainPropsRead.h"
 #include "PlainPropsInternalRead.h"
 #include "Misc/Optional.h"
@@ -442,4 +442,5 @@ FFlatMemberReader::FFlatMemberReader(FStructView Struct)
 	It = &Lineage.Last();
 }
 
-} // namespace PlainProps
\ No newline at end of file
+} // namespace PlainProps
+#endif
\ No newline at end of file
diff --git a/Engine/Plugins/Experimental/PlainProps/Source/Private/PlainPropsSave.cpp b/Engine/Plugins/Experimental/PlainProps/Source/Private/PlainPropsSave.cpp
--- a/Engine/Plugins/Experimental/PlainProps/Source/Private/PlainPropsSave.cpp
+++ b/Engine/Plugins/Experimental/PlainProps/Source/Private/PlainPropsSave.cpp
@@ -1,5 +1,5 @@
 // Copyright Epic Games, Inc. All Rights Reserved.
-
+#if 0
 #include "PlainPropsSave.h"
 #include "PlainPropsBind.h"
 #include "PlainPropsInternalBuild.h"
@@ -473,4 +473,5 @@ TUniquePtr<FBuiltStruct> SaveStructDelta(const void* Struct, const void* Default
 	return !Out.IsEmpty() ? Out.BuildAndReset(Declaration, Ctx.Declarations.GetDebug()) : nullptr;
 }
 
-} // namespace PlainProps
\ No newline at end of file
+} // namespace PlainProps
+#endif
\ No newline at end of file
diff --git a/Engine/Plugins/Experimental/PlainProps/Source/Private/PlainPropsWrite.cpp b/Engine/Plugins/Experimental/PlainProps/Source/Private/PlainPropsWrite.cpp
--- a/Engine/Plugins/Experimental/PlainProps/Source/Private/PlainPropsWrite.cpp
+++ b/Engine/Plugins/Experimental/PlainProps/Source/Private/PlainPropsWrite.cpp
@@ -1,5 +1,5 @@
 // Copyright Epic Games, Inc. All Rights Reserved.
-
+#if 0
 #include "PlainPropsWrite.h"
 #include "PlainPropsBuildSchema.h"
 #include "PlainPropsIndex.h"
@@ -859,3 +859,4 @@ FString FDebugIds::Print(FStructSchemaId Name) const
 }
 
 } // namespace PlainProps
+#endif
\ No newline at end of file
diff --git a/Engine/Plugins/Experimental/PlainProps/Source/Tests/PlainPropsCttiTests.cpp b/Engine/Plugins/Experimental/PlainProps/Source/Tests/PlainPropsCttiTests.cpp
--- a/Engine/Plugins/Experimental/PlainProps/Source/Tests/PlainPropsCttiTests.cpp
+++ b/Engine/Plugins/Experimental/PlainProps/Source/Tests/PlainPropsCttiTests.cpp
@@ -1,5 +1,5 @@
 // Copyright Epic Games, Inc. All Rights Reserved.
-
+#if 0
 #include "PlainPropsCtti.h"
 #include <type_traits>
 #include <string_view>
@@ -123,4 +123,5 @@ static_assert(CttiOf<S2<int>>::NumVars == 1);
 static_assert(CttiOf<S2<int>>::Var<0>::Name == std::string_view("a"));
 static_assert(CttiOf<S2<int>>::Var<0>::Offset == offsetof(S2<int>, a));
 
-}
\ No newline at end of file
+}
+#endif
\ No newline at end of file
diff --git a/Engine/Plugins/Experimental/PlainProps/Source/Tests/PlainPropsTests.cpp b/Engine/Plugins/Experimental/PlainProps/Source/Tests/PlainPropsTests.cpp
--- a/Engine/Plugins/Experimental/PlainProps/Source/Tests/PlainPropsTests.cpp
+++ b/Engine/Plugins/Experimental/PlainProps/Source/Tests/PlainPropsTests.cpp
@@ -1,5 +1,5 @@
 // Copyright Epic Games, Inc. All Rights Reserved.
-
+#if 0
 #if WITH_TESTS
 #include "Tests/TestHarnessAdapter.h"
 #include "Containers/AnsiString.h"
@@ -1350,4 +1350,5 @@ TEST_CASE_NAMED(FPlainPropsLoadSaveTest, "System::Core::Serialization::PlainProp
 }
 
 } // namespace PlainProps
-#endif // WITH_TESTS
\ No newline at end of file
+#endif // WITH_TESTS
+#endif
\ No newline at end of file
diff --git a/Engine/Plugins/Experimental/PlainProps/Source/Tests/PlainPropsUeCoreTests.cpp b/Engine/Plugins/Experimental/PlainProps/Source/Tests/PlainPropsUeCoreTests.cpp
--- a/Engine/Plugins/Experimental/PlainProps/Source/Tests/PlainPropsUeCoreTests.cpp
+++ b/Engine/Plugins/Experimental/PlainProps/Source/Tests/PlainPropsUeCoreTests.cpp
@@ -1,5 +1,5 @@
 // Copyright Epic Games, Inc. All Rights Reserved.
-
+#if 0
 #if WITH_TESTS
 
 #include "PlainPropsBuildSchema.h"
@@ -761,4 +761,5 @@ TEST_CASE_NAMED(FPlainPropsUeCoreTest, "System::Core::Serialization::PlainProps:
 }
 
 } // namespace PlainProps::UE::Test
-#endif // WITH_TESTS
\ No newline at end of file
+#endif // WITH_TESTS
+#endif
\ No newline at end of file
