Another temp fix for build jobs that compile standalone .h files with C++17, sigh

#rnx

diff --git a/Engine/Plugins/Experimental/PlainProps/Source/Private/PlainPropsInternalBuild.h b/Engine/Plugins/Experimental/PlainProps/Source/Private/PlainPropsInternalBuild.h
--- a/Engine/Plugins/Experimental/PlainProps/Source/Private/PlainPropsInternalBuild.h
+++ b/Engine/Plugins/Experimental/PlainProps/Source/Private/PlainPropsInternalBuild.h
@@ -1,5 +1,5 @@
 // Copyright Epic Games, Inc. All Rights Reserved.
-
+#if 0
 #pragma once
 
 #include "PlainPropsBuild.h"
@@ -101,4 +101,5 @@ inline void WriteU64(TArray64<uint8>& Out, uint64 Int) { WriteInt(Out, Int); }
 PLAINPROPS_API uint64 WriteSkippableSlice(TArray64<uint8>& Out, TConstArrayView64<uint8> Slice);
 
 
-} // namespace PlainProps
\ No newline at end of file
+} // namespace PlainProps
+#endif
\ No newline at end of file
diff --git a/Engine/Plugins/Experimental/PlainProps/Source/Private/PlainPropsInternalFormat.h b/Engine/Plugins/Experimental/PlainProps/Source/Private/PlainPropsInternalFormat.h
--- a/Engine/Plugins/Experimental/PlainProps/Source/Private/PlainPropsInternalFormat.h
+++ b/Engine/Plugins/Experimental/PlainProps/Source/Private/PlainPropsInternalFormat.h
@@ -1,5 +1,5 @@
 // Copyright Epic Games, Inc. All Rights Reserved.
-
+#if 0
 #pragma once
 
 #include "PlainPropsTypes.h"
@@ -157,4 +157,5 @@ inline constexpr uint64 GetLeafRangeSize(uint64 Num, FUnpackedLeafType Leaf)
 
 //////////////////////////////////////////////////////////////////////////
 
-} // namespace PlainProps
\ No newline at end of file
+} // namespace PlainProps
+#endif
\ No newline at end of file
diff --git a/Engine/Plugins/Experimental/PlainProps/Source/Private/PlainPropsInternalRead.h b/Engine/Plugins/Experimental/PlainProps/Source/Private/PlainPropsInternalRead.h
--- a/Engine/Plugins/Experimental/PlainProps/Source/Private/PlainPropsInternalRead.h
+++ b/Engine/Plugins/Experimental/PlainProps/Source/Private/PlainPropsInternalRead.h
@@ -1,5 +1,5 @@
 // Copyright Epic Games, Inc. All Rights Reserved.
-
+#if 0
 #pragma once
 
 #include "PlainPropsTypes.h"
@@ -93,3 +93,4 @@ inline FMemoryView GetSchemaData(const FSchemaBatch& Batch)
 }
 
 } // namespace PlainProps
+#endif
\ No newline at end of file
diff --git a/Engine/Plugins/Experimental/PlainProps/Source/Public/PlainPropsBind.h b/Engine/Plugins/Experimental/PlainProps/Source/Public/PlainPropsBind.h
--- a/Engine/Plugins/Experimental/PlainProps/Source/Public/PlainPropsBind.h
+++ b/Engine/Plugins/Experimental/PlainProps/Source/Public/PlainPropsBind.h
@@ -1,5 +1,5 @@
 // Copyright Epic Games, Inc. All Rights Reserved.
-
+#if 0
 #pragma once
 
 #include "Containers/ArrayView.h"
@@ -889,4 +889,5 @@ struct FIdTranslator : FIdTranslatorBase
 FSchemaBatch*			CreateTranslatedSchemas(const FSchemaBatch& Schemas, FIdBinding NewIds);
 void					DestroyTranslatedSchemas(const FSchemaBatch* Schemas);
 
-} // namespace PlainProps
\ No newline at end of file
+} // namespace PlainProps
+#endif
\ No newline at end of file
diff --git a/Engine/Plugins/Experimental/PlainProps/Source/Public/PlainPropsBuild.h b/Engine/Plugins/Experimental/PlainProps/Source/Public/PlainPropsBuild.h
--- a/Engine/Plugins/Experimental/PlainProps/Source/Public/PlainPropsBuild.h
+++ b/Engine/Plugins/Experimental/PlainProps/Source/Public/PlainPropsBuild.h
@@ -1,5 +1,5 @@
 // Copyright Epic Games, Inc. All Rights Reserved.
-
+#if 0
 #pragma once
 
 #include "PlainPropsDeclare.h"
@@ -240,4 +240,5 @@ private:
 	FMemberSchema Schema;
 };
 
-} // namespace PlainProps
\ No newline at end of file
+} // namespace PlainProps
+#endif
\ No newline at end of file
diff --git a/Engine/Plugins/Experimental/PlainProps/Source/Public/PlainPropsBuildSchema.h b/Engine/Plugins/Experimental/PlainProps/Source/Public/PlainPropsBuildSchema.h
--- a/Engine/Plugins/Experimental/PlainProps/Source/Public/PlainPropsBuildSchema.h
+++ b/Engine/Plugins/Experimental/PlainProps/Source/Public/PlainPropsBuildSchema.h
@@ -1,5 +1,5 @@
 // Copyright Epic Games, Inc. All Rights Reserved.
-
+#if 0
 #pragma once
 
 #include "PlainPropsBuild.h"
@@ -67,4 +67,5 @@ private:
 	void										NoteInheritanceChains();
 };
 
-} // namespace PlainProps
\ No newline at end of file
+} // namespace PlainProps
+#endif
\ No newline at end of file
diff --git a/Engine/Plugins/Experimental/PlainProps/Source/Public/PlainPropsCtti.h b/Engine/Plugins/Experimental/PlainProps/Source/Public/PlainPropsCtti.h
--- a/Engine/Plugins/Experimental/PlainProps/Source/Public/PlainPropsCtti.h
+++ b/Engine/Plugins/Experimental/PlainProps/Source/Public/PlainPropsCtti.h
@@ -1,5 +1,5 @@
 // Copyright Epic Games, Inc. All Rights Reserved.
-
+#if 0
 #pragma once
 
 #include <cstdint>
@@ -92,3 +92,4 @@ static constexpr void ForEachVar(Fn&& Visitor)
 
 } // namespace PlainProps
 
+#endif
\ No newline at end of file
diff --git a/Engine/Plugins/Experimental/PlainProps/Source/Public/PlainPropsDeclare.h b/Engine/Plugins/Experimental/PlainProps/Source/Public/PlainPropsDeclare.h
--- a/Engine/Plugins/Experimental/PlainProps/Source/Public/PlainPropsDeclare.h
+++ b/Engine/Plugins/Experimental/PlainProps/Source/Public/PlainPropsDeclare.h
@@ -1,5 +1,5 @@
 // Copyright Epic Games, Inc. All Rights Reserved.
-
+#if 0
 #pragma once
 
 #include "PlainPropsTypes.h"
@@ -80,4 +80,5 @@ protected:
 #endif
 };
 
-} // namespace PlainProps
\ No newline at end of file
+} // namespace PlainProps
+#endif
\ No newline at end of file
diff --git a/Engine/Plugins/Experimental/PlainProps/Source/Public/PlainPropsIndex.h b/Engine/Plugins/Experimental/PlainProps/Source/Public/PlainPropsIndex.h
--- a/Engine/Plugins/Experimental/PlainProps/Source/Public/PlainPropsIndex.h
+++ b/Engine/Plugins/Experimental/PlainProps/Source/Public/PlainPropsIndex.h
@@ -1,5 +1,5 @@
 // Copyright Epic Games, Inc. All Rights Reserved.
-
+#if 0
 #pragma once
 
 #include "PlainPropsTypes.h"
@@ -106,3 +106,4 @@ protected:
 
 
 } // namespace PlainProps
+#endif
\ No newline at end of file
diff --git a/Engine/Plugins/Experimental/PlainProps/Source/Public/PlainPropsLoad.h b/Engine/Plugins/Experimental/PlainProps/Source/Public/PlainPropsLoad.h
--- a/Engine/Plugins/Experimental/PlainProps/Source/Public/PlainPropsLoad.h
+++ b/Engine/Plugins/Experimental/PlainProps/Source/Public/PlainPropsLoad.h
@@ -1,5 +1,5 @@
 // Copyright Epic Games, Inc. All Rights Reserved.
-
+#if 0
 #pragma once
 
 #include "Containers/ArrayView.h"
@@ -39,3 +39,4 @@ PLAINPROPS_API void LoadRange(uint8* Dst, FRangeView Src, ERangeSizeType MaxType
 
 } // namespace PlainProps
 
+#endif
\ No newline at end of file
diff --git a/Engine/Plugins/Experimental/PlainProps/Source/Public/PlainPropsRead.h b/Engine/Plugins/Experimental/PlainProps/Source/Public/PlainPropsRead.h
--- a/Engine/Plugins/Experimental/PlainProps/Source/Public/PlainPropsRead.h
+++ b/Engine/Plugins/Experimental/PlainProps/Source/Public/PlainPropsRead.h
@@ -1,5 +1,5 @@
 // Copyright Epic Games, Inc. All Rights Reserved.
-
+#if 0
 #pragma once
 
 #include "Containers/ArrayView.h"
@@ -616,4 +616,5 @@ private:
 	}
 };
 
-} // namespace PlainProps
\ No newline at end of file
+} // namespace PlainProps
+#endif
\ No newline at end of file
diff --git a/Engine/Plugins/Experimental/PlainProps/Source/Public/PlainPropsSave.h b/Engine/Plugins/Experimental/PlainProps/Source/Public/PlainPropsSave.h
--- a/Engine/Plugins/Experimental/PlainProps/Source/Public/PlainPropsSave.h
+++ b/Engine/Plugins/Experimental/PlainProps/Source/Public/PlainPropsSave.h
@@ -1,5 +1,5 @@
 // Copyright Epic Games, Inc. All Rights Reserved.
-
+#if 0
 #pragma once
 
 #include "PlainPropsTypes.h"
@@ -25,3 +25,4 @@ struct FSaveContext
 
 } // namespace PlainProps
 
+#endif
\ No newline at end of file
diff --git a/Engine/Plugins/Experimental/PlainProps/Source/Public/PlainPropsTypes.h b/Engine/Plugins/Experimental/PlainProps/Source/Public/PlainPropsTypes.h
--- a/Engine/Plugins/Experimental/PlainProps/Source/Public/PlainPropsTypes.h
+++ b/Engine/Plugins/Experimental/PlainProps/Source/Public/PlainPropsTypes.h
@@ -1,5 +1,5 @@
 // Copyright Epic Games, Inc. All Rights Reserved.
-
+#if 0
 #pragma once
 
 #include "Containers/ArrayView.h"
@@ -381,4 +381,5 @@ public:
 	PLAINPROPS_API FString				Print(FStructSchemaId Name) const;
 };
 
-} // namespace PlainProps
\ No newline at end of file
+} // namespace PlainProps
+#endif
\ No newline at end of file
diff --git a/Engine/Plugins/Experimental/PlainProps/Source/Public/PlainPropsUeCoreBindings.h b/Engine/Plugins/Experimental/PlainProps/Source/Public/PlainPropsUeCoreBindings.h
--- a/Engine/Plugins/Experimental/PlainProps/Source/Public/PlainPropsUeCoreBindings.h
+++ b/Engine/Plugins/Experimental/PlainProps/Source/Public/PlainPropsUeCoreBindings.h
@@ -1,5 +1,5 @@
 // Copyright Epic Games, Inc. All Rights Reserved.
-
+#if 0
 #pragma once
 
 #include "Containers/Set.h"
@@ -590,4 +590,5 @@ namespace PlainProps
 	{
 		using Type = UE::TSetBinding<T>;
 	};
-}
\ No newline at end of file
+}
+#endif
\ No newline at end of file
diff --git a/Engine/Plugins/Experimental/PlainProps/Source/Public/PlainPropsWrite.h b/Engine/Plugins/Experimental/PlainProps/Source/Public/PlainPropsWrite.h
--- a/Engine/Plugins/Experimental/PlainProps/Source/Public/PlainPropsWrite.h
+++ b/Engine/Plugins/Experimental/PlainProps/Source/Public/PlainPropsWrite.h
@@ -1,5 +1,5 @@
 // Copyright Epic Games, Inc. All Rights Reserved.
-
+#if 0
 #pragma once
 
 #include "PlainPropsTypes.h"
@@ -36,4 +36,5 @@ private:
 	TUniquePtr<FWriteIds>							NewIds;
 };
 
-} // namespace PlainProps
\ No newline at end of file
+} // namespace PlainProps
+#endif
\ No newline at end of file
