PlainProps undo temp hack, fix some unity specific compile errors, a linux and a mac issue

Also disable platforms we don't use currently to avoid pointless breakages

#rnx

diff --git a/Engine/Plugins/Experimental/PlainProps/Source/PlainProps.Build.cs b/Engine/Plugins/Experimental/PlainProps/Source/PlainProps.Build.cs
--- a/Engine/Plugins/Experimental/PlainProps/Source/PlainProps.Build.cs
+++ b/Engine/Plugins/Experimental/PlainProps/Source/PlainProps.Build.cs
@@ -2,11 +2,13 @@
 
 using UnrealBuildTool;
 
-[SupportedPlatforms(UnrealPlatformClass.All)]
+[SupportedPlatforms("Win64", "PS5")]
 public class PlainProps : ModuleRules
 {
 	public PlainProps(ReadOnlyTargetRules Target) : base(Target)
 	{
+		CppStandard = CppStandardVersion.Cpp20;
+
 		PrivateDependencyModuleNames.AddRange(
 			new string[] {
 				"Core",
diff --git a/Engine/Plugins/Experimental/PlainProps/Source/Private/PlainPropsBind.cpp b/Engine/Plugins/Experimental/PlainProps/Source/Private/PlainPropsBind.cpp
--- a/Engine/Plugins/Experimental/PlainProps/Source/Private/PlainPropsBind.cpp
+++ b/Engine/Plugins/Experimental/PlainProps/Source/Private/PlainPropsBind.cpp
@@ -1,5 +1,5 @@
 // Copyright Epic Games, Inc. All Rights Reserved.
-#if 0
+
 #include "PlainPropsBind.h"
 #include "PlainPropsIndex.h"
 #include "PlainPropsInternalBuild.h"
@@ -376,5 +376,4 @@ void DestroyTranslatedSchemas(const FSchemaBatch* Schemas)
 	FMemory::Free(const_cast<FSchemaBatch*>(Schemas));
 }
 
-} // namespace PlainProps
-#endif
\ No newline at end of file
+} // namespace PlainProps
\ No newline at end of file
diff --git a/Engine/Plugins/Experimental/PlainProps/Source/Private/PlainPropsBuild.cpp b/Engine/Plugins/Experimental/PlainProps/Source/Private/PlainPropsBuild.cpp
--- a/Engine/Plugins/Experimental/PlainProps/Source/Private/PlainPropsBuild.cpp
+++ b/Engine/Plugins/Experimental/PlainProps/Source/Private/PlainPropsBuild.cpp
@@ -1,5 +1,5 @@
 // Copyright Epic Games, Inc. All Rights Reserved.
-#if 0
+
 #include "PlainPropsBuild.h"
 #include "PlainPropsIndex.h"
 #include "PlainPropsInternalBuild.h"
@@ -303,5 +303,4 @@ FTypedRange FNestedRangeBuilder::BuildAndReset(ERangeSizeType SizeType)
 	return Out;
 }
 
-} // namespace PlainProps
-#endif
\ No newline at end of file
+} // namespace PlainProps
\ No newline at end of file
diff --git a/Engine/Plugins/Experimental/PlainProps/Source/Private/PlainPropsBuildSchema.cpp b/Engine/Plugins/Experimental/PlainProps/Source/Private/PlainPropsBuildSchema.cpp
--- a/Engine/Plugins/Experimental/PlainProps/Source/Private/PlainPropsBuildSchema.cpp
+++ b/Engine/Plugins/Experimental/PlainProps/Source/Private/PlainPropsBuildSchema.cpp
@@ -1,5 +1,5 @@
 // Copyright Epic Games, Inc. All Rights Reserved.
-#if 0
+
 #include "PlainPropsBuildSchema.h"
 #include "PlainPropsInternalBuild.h"
 #include "PlainPropsInternalFormat.h"
@@ -232,12 +232,6 @@ void FStructSchemaBuilder::NoteMembersRecursively(const FBuiltStruct& Struct)
 	}
 }
 
-template<class T>
-TConstArrayView64<T> MakeArrayView(const T* Data, uint64 Num)
-{
-	return TConstArrayView64<T>(Data, static_cast<int64>(Num));
-}
-
 template<typename IntType>
 void NoteEnumValues(FEnumSchemaBuilder& Schema, const IntType* Values, uint64 Num)
 {
@@ -387,4 +381,3 @@ void FEnumSchemaBuilder::NoteValue(uint64 Value)
 }
 
 } // namespace PlainProps
-#endif
\ No newline at end of file
diff --git a/Engine/Plugins/Experimental/PlainProps/Source/Private/PlainPropsDeclare.cpp b/Engine/Plugins/Experimental/PlainProps/Source/Private/PlainPropsDeclare.cpp
--- a/Engine/Plugins/Experimental/PlainProps/Source/Private/PlainPropsDeclare.cpp
+++ b/Engine/Plugins/Experimental/PlainProps/Source/Private/PlainPropsDeclare.cpp
@@ -1,5 +1,5 @@
 // Copyright Epic Games, Inc. All Rights Reserved.
-#if 0
+
 #include "PlainPropsDeclare.h"
 #include "Containers/Set.h"
 
@@ -89,5 +89,4 @@ void FDeclarations::Check(FStructSchemaId Id) const
 }
 #endif
 
-} // namespace PlainProps
-#endif
\ No newline at end of file
+} // namespace PlainProps
\ No newline at end of file
diff --git a/Engine/Plugins/Experimental/PlainProps/Source/Private/PlainPropsIndex.cpp b/Engine/Plugins/Experimental/PlainProps/Source/Private/PlainPropsIndex.cpp
--- a/Engine/Plugins/Experimental/PlainProps/Source/Private/PlainPropsIndex.cpp
+++ b/Engine/Plugins/Experimental/PlainProps/Source/Private/PlainPropsIndex.cpp
@@ -1,5 +1,5 @@
 // Copyright Epic Games, Inc. All Rights Reserved.
-#if 0
+
 #include "PlainPropsIndex.h"
 #include "Hash/xxhash.h"
 
@@ -147,5 +147,4 @@ FStructSchemaId	FIdIndexerBase::IndexStruct(FTypeId Type)
 	return { IntCastChecked<uint32>(Structs.Add(Type).AsInteger()) };
 }
 
-} // namespace PlainProps
-#endif
\ No newline at end of file
+} // namespace PlainProps
\ No newline at end of file
diff --git a/Engine/Plugins/Experimental/PlainProps/Source/Private/PlainPropsInternalBuild.h b/Engine/Plugins/Experimental/PlainProps/Source/Private/PlainPropsInternalBuild.h
--- a/Engine/Plugins/Experimental/PlainProps/Source/Private/PlainPropsInternalBuild.h
+++ b/Engine/Plugins/Experimental/PlainProps/Source/Private/PlainPropsInternalBuild.h
@@ -1,5 +1,5 @@
 // Copyright Epic Games, Inc. All Rights Reserved.
-#if 0
+
 #pragma once
 
 #include "PlainPropsBuild.h"
@@ -101,5 +101,4 @@ inline void WriteU64(TArray64<uint8>& Out, uint64 Int) { WriteInt(Out, Int); }
 PLAINPROPS_API uint64 WriteSkippableSlice(TArray64<uint8>& Out, TConstArrayView64<uint8> Slice);
 
 
-} // namespace PlainProps
-#endif
\ No newline at end of file
+} // namespace PlainProps
\ No newline at end of file
diff --git a/Engine/Plugins/Experimental/PlainProps/Source/Private/PlainPropsInternalFormat.h b/Engine/Plugins/Experimental/PlainProps/Source/Private/PlainPropsInternalFormat.h
--- a/Engine/Plugins/Experimental/PlainProps/Source/Private/PlainPropsInternalFormat.h
+++ b/Engine/Plugins/Experimental/PlainProps/Source/Private/PlainPropsInternalFormat.h
@@ -1,5 +1,5 @@
 // Copyright Epic Games, Inc. All Rights Reserved.
-#if 0
+
 #pragma once
 
 #include "PlainPropsTypes.h"
@@ -157,5 +157,4 @@ inline constexpr uint64 GetLeafRangeSize(uint64 Num, FUnpackedLeafType Leaf)
 
 //////////////////////////////////////////////////////////////////////////
 
-} // namespace PlainProps
-#endif
\ No newline at end of file
+} // namespace PlainProps
\ No newline at end of file
diff --git a/Engine/Plugins/Experimental/PlainProps/Source/Private/PlainPropsInternalRead.h b/Engine/Plugins/Experimental/PlainProps/Source/Private/PlainPropsInternalRead.h
--- a/Engine/Plugins/Experimental/PlainProps/Source/Private/PlainPropsInternalRead.h
+++ b/Engine/Plugins/Experimental/PlainProps/Source/Private/PlainPropsInternalRead.h
@@ -1,5 +1,5 @@
 // Copyright Epic Games, Inc. All Rights Reserved.
-#if 0
+
 #pragma once
 
 #include "PlainPropsTypes.h"
@@ -93,4 +93,3 @@ inline FMemoryView GetSchemaData(const FSchemaBatch& Batch)
 }
 
 } // namespace PlainProps
-#endif
\ No newline at end of file
diff --git a/Engine/Plugins/Experimental/PlainProps/Source/Private/PlainPropsLoad.cpp b/Engine/Plugins/Experimental/PlainProps/Source/Private/PlainPropsLoad.cpp
--- a/Engine/Plugins/Experimental/PlainProps/Source/Private/PlainPropsLoad.cpp
+++ b/Engine/Plugins/Experimental/PlainProps/Source/Private/PlainPropsLoad.cpp
@@ -1,5 +1,5 @@
 // Copyright Epic Games, Inc. All Rights Reserved.
-#if 0
+
 #include "PlainPropsLoad.h"
 #include "PlainPropsBind.h"
 #include "PlainPropsInternalFormat.h"
@@ -646,5 +646,5 @@ void LoadRange(uint8* Dst, FRangeView Src, ERangeSizeType MaxSize, TConstArrayVi
 	FRangeLoader::LoadRangeView(Dst, Src, MaxSize, Bindings, Batch);
 }
 
-} // namespace PlainProps
-#endif
\ No newline at end of file
+
+} // namespace PlainProps
\ No newline at end of file
diff --git a/Engine/Plugins/Experimental/PlainProps/Source/Private/PlainPropsRead.cpp b/Engine/Plugins/Experimental/PlainProps/Source/Private/PlainPropsRead.cpp
--- a/Engine/Plugins/Experimental/PlainProps/Source/Private/PlainPropsRead.cpp
+++ b/Engine/Plugins/Experimental/PlainProps/Source/Private/PlainPropsRead.cpp
@@ -1,5 +1,5 @@
 // Copyright Epic Games, Inc. All Rights Reserved.
-#if 0
+
 #include "PlainPropsRead.h"
 #include "PlainPropsInternalRead.h"
 #include "Misc/Optional.h"
@@ -442,5 +442,4 @@ FFlatMemberReader::FFlatMemberReader(FStructView Struct)
 	It = &Lineage.Last();
 }
 
-} // namespace PlainProps
-#endif
\ No newline at end of file
+} // namespace PlainProps
\ No newline at end of file
diff --git a/Engine/Plugins/Experimental/PlainProps/Source/Private/PlainPropsSave.cpp b/Engine/Plugins/Experimental/PlainProps/Source/Private/PlainPropsSave.cpp
--- a/Engine/Plugins/Experimental/PlainProps/Source/Private/PlainPropsSave.cpp
+++ b/Engine/Plugins/Experimental/PlainProps/Source/Private/PlainPropsSave.cpp
@@ -1,5 +1,5 @@
 // Copyright Epic Games, Inc. All Rights Reserved.
-#if 0
+
 #include "PlainPropsSave.h"
 #include "PlainPropsBind.h"
 #include "PlainPropsInternalBuild.h"
@@ -473,5 +473,4 @@ TUniquePtr<FBuiltStruct> SaveStructDelta(const void* Struct, const void* Default
 	return !Out.IsEmpty() ? Out.BuildAndReset(Declaration, Ctx.Declarations.GetDebug()) : nullptr;
 }
 
-} // namespace PlainProps
-#endif
\ No newline at end of file
+} // namespace PlainProps
\ No newline at end of file
diff --git a/Engine/Plugins/Experimental/PlainProps/Source/Private/PlainPropsWrite.cpp b/Engine/Plugins/Experimental/PlainProps/Source/Private/PlainPropsWrite.cpp
--- a/Engine/Plugins/Experimental/PlainProps/Source/Private/PlainPropsWrite.cpp
+++ b/Engine/Plugins/Experimental/PlainProps/Source/Private/PlainPropsWrite.cpp
@@ -1,5 +1,5 @@
 // Copyright Epic Games, Inc. All Rights Reserved.
-#if 0
+
 #include "PlainPropsWrite.h"
 #include "PlainPropsBuildSchema.h"
 #include "PlainPropsIndex.h"
@@ -859,4 +859,3 @@ FString FDebugIds::Print(FStructSchemaId Name) const
 }
 
 } // namespace PlainProps
-#endif
\ No newline at end of file
diff --git a/Engine/Plugins/Experimental/PlainProps/Source/Public/PlainPropsBind.h b/Engine/Plugins/Experimental/PlainProps/Source/Public/PlainPropsBind.h
--- a/Engine/Plugins/Experimental/PlainProps/Source/Public/PlainPropsBind.h
+++ b/Engine/Plugins/Experimental/PlainProps/Source/Public/PlainPropsBind.h
@@ -1,5 +1,5 @@
 // Copyright Epic Games, Inc. All Rights Reserved.
-#if 0
+
 #pragma once
 
 #include "Containers/ArrayView.h"
@@ -889,5 +889,4 @@ struct FIdTranslator : FIdTranslatorBase
 FSchemaBatch*			CreateTranslatedSchemas(const FSchemaBatch& Schemas, FIdBinding NewIds);
 void					DestroyTranslatedSchemas(const FSchemaBatch* Schemas);
 
-} // namespace PlainProps
-#endif
\ No newline at end of file
+} // namespace PlainProps
\ No newline at end of file
diff --git a/Engine/Plugins/Experimental/PlainProps/Source/Public/PlainPropsBuild.h b/Engine/Plugins/Experimental/PlainProps/Source/Public/PlainPropsBuild.h
--- a/Engine/Plugins/Experimental/PlainProps/Source/Public/PlainPropsBuild.h
+++ b/Engine/Plugins/Experimental/PlainProps/Source/Public/PlainPropsBuild.h
@@ -1,5 +1,5 @@
 // Copyright Epic Games, Inc. All Rights Reserved.
-#if 0
+
 #pragma once
 
 #include "PlainPropsDeclare.h"
@@ -49,16 +49,6 @@ uint64 ValueCast(double Value);
 inline uint64 ValueCast(char8_t Value)		{ return static_cast<uint8>(Value); }
 inline uint64 ValueCast(char16_t Value)		{ return static_cast<uint16>(Value); }
 inline uint64 ValueCast(char32_t Value)		{ return static_cast<uint32>(Value); }
-	
-inline constexpr ERangeSizeType RangeSizeOf(bool)	{ return ERangeSizeType::Uni; }
-inline constexpr ERangeSizeType RangeSizeOf(int8)	{ return ERangeSizeType::S8; }
-inline constexpr ERangeSizeType RangeSizeOf(int16)	{ return ERangeSizeType::S16; }
-inline constexpr ERangeSizeType RangeSizeOf(int32)	{ return ERangeSizeType::S32; }
-inline constexpr ERangeSizeType RangeSizeOf(int64)	{ return ERangeSizeType::S64; }
-inline constexpr ERangeSizeType RangeSizeOf(uint8)	{ return ERangeSizeType::U8; }
-inline constexpr ERangeSizeType RangeSizeOf(uint16)	{ return ERangeSizeType::U16; }
-inline constexpr ERangeSizeType RangeSizeOf(uint32)	{ return ERangeSizeType::U32; }
-inline constexpr ERangeSizeType RangeSizeOf(uint64)	{ return ERangeSizeType::U64; }
 
 //////////////////////////////////////////////////////////////////////////
 
@@ -97,9 +87,9 @@ PLAINPROPS_API FMemberSchema MakeNestedRangeSchema(ERangeSizeType SizeType, cons
 
 namespace Private
 {
-	PLAINPROPS_API [[nodiscard]] FBuiltRange*	BuildStructuralRange(/* in-out */ TArrayView64<FBuiltRange*> Structs);
-	PLAINPROPS_API [[nodiscard]] FBuiltRange*	BuildStructuralRange(/* in-out */ TArrayView64<TUniquePtr<FBuiltStruct>> Structs);
-	PLAINPROPS_API [[nodiscard]] FBuiltRange*	BuildLeafRange(FUnpackedLeafType Leaf, uint64 Num, FMemoryView Values);
+	[[nodiscard]] PLAINPROPS_API FBuiltRange*	BuildStructuralRange(/* in-out */ TArrayView64<FBuiltRange*> Structs);
+	[[nodiscard]] PLAINPROPS_API FBuiltRange*	BuildStructuralRange(/* in-out */ TArrayView64<TUniquePtr<FBuiltStruct>> Structs);
+	[[nodiscard]] PLAINPROPS_API FBuiltRange*	BuildLeafRange(FUnpackedLeafType Leaf, uint64 Num, FMemoryView Values);
 	PLAINPROPS_API void							NormalizeLeafRange(FUnpackedLeafType Leaf, FBuiltRange& Out);
 
 	template<typename LeafType, typename SizeType>
@@ -174,7 +164,7 @@ public:
 	// Build members into a single nested super struct member, no-op if no non-super members has been added
 	PLAINPROPS_API void BuildSuperStruct(const FStructDeclaration& Super, const FDebugIds& Debug);
 
-	PLAINPROPS_API [[nodiscard]] TUniquePtr<FBuiltStruct> BuildAndReset(const FStructDeclaration& Declared, const FDebugIds& Debug);
+	[[nodiscard]] PLAINPROPS_API TUniquePtr<FBuiltStruct> BuildAndReset(const FStructDeclaration& Declared, const FDebugIds& Debug);
 
 	bool IsEmpty() const { return Members.IsEmpty(); }
 
@@ -240,5 +230,4 @@ private:
 	FMemberSchema Schema;
 };
 
-} // namespace PlainProps
-#endif
\ No newline at end of file
+} // namespace PlainProps
\ No newline at end of file
diff --git a/Engine/Plugins/Experimental/PlainProps/Source/Public/PlainPropsBuildSchema.h b/Engine/Plugins/Experimental/PlainProps/Source/Public/PlainPropsBuildSchema.h
--- a/Engine/Plugins/Experimental/PlainProps/Source/Public/PlainPropsBuildSchema.h
+++ b/Engine/Plugins/Experimental/PlainProps/Source/Public/PlainPropsBuildSchema.h
@@ -1,5 +1,5 @@
 // Copyright Epic Games, Inc. All Rights Reserved.
-#if 0
+
 #pragma once
 
 #include "PlainPropsBuild.h"
@@ -67,5 +67,4 @@ private:
 	void										NoteInheritanceChains();
 };
 
-} // namespace PlainProps
-#endif
\ No newline at end of file
+} // namespace PlainProps
\ No newline at end of file
diff --git a/Engine/Plugins/Experimental/PlainProps/Source/Public/PlainPropsCtti.h b/Engine/Plugins/Experimental/PlainProps/Source/Public/PlainPropsCtti.h
--- a/Engine/Plugins/Experimental/PlainProps/Source/Public/PlainPropsCtti.h
+++ b/Engine/Plugins/Experimental/PlainProps/Source/Public/PlainPropsCtti.h
@@ -1,5 +1,5 @@
 // Copyright Epic Games, Inc. All Rights Reserved.
-#if 0
+
 #pragma once
 
 #include <cstdint>
@@ -92,4 +92,3 @@ static constexpr void ForEachVar(Fn&& Visitor)
 
 } // namespace PlainProps
 
-#endif
\ No newline at end of file
diff --git a/Engine/Plugins/Experimental/PlainProps/Source/Public/PlainPropsDeclare.h b/Engine/Plugins/Experimental/PlainProps/Source/Public/PlainPropsDeclare.h
--- a/Engine/Plugins/Experimental/PlainProps/Source/Public/PlainPropsDeclare.h
+++ b/Engine/Plugins/Experimental/PlainProps/Source/Public/PlainPropsDeclare.h
@@ -1,5 +1,5 @@
 // Copyright Epic Games, Inc. All Rights Reserved.
-#if 0
+
 #pragma once
 
 #include "PlainPropsTypes.h"
@@ -80,5 +80,4 @@ protected:
 #endif
 };
 
-} // namespace PlainProps
-#endif
\ No newline at end of file
+} // namespace PlainProps
\ No newline at end of file
diff --git a/Engine/Plugins/Experimental/PlainProps/Source/Public/PlainPropsIndex.h b/Engine/Plugins/Experimental/PlainProps/Source/Public/PlainPropsIndex.h
--- a/Engine/Plugins/Experimental/PlainProps/Source/Public/PlainPropsIndex.h
+++ b/Engine/Plugins/Experimental/PlainProps/Source/Public/PlainPropsIndex.h
@@ -1,5 +1,5 @@
 // Copyright Epic Games, Inc. All Rights Reserved.
-#if 0
+
 #pragma once
 
 #include "PlainPropsTypes.h"
@@ -106,4 +106,3 @@ protected:
 
 
 } // namespace PlainProps
-#endif
\ No newline at end of file
diff --git a/Engine/Plugins/Experimental/PlainProps/Source/Public/PlainPropsLoad.h b/Engine/Plugins/Experimental/PlainProps/Source/Public/PlainPropsLoad.h
--- a/Engine/Plugins/Experimental/PlainProps/Source/Public/PlainPropsLoad.h
+++ b/Engine/Plugins/Experimental/PlainProps/Source/Public/PlainPropsLoad.h
@@ -1,5 +1,5 @@
 // Copyright Epic Games, Inc. All Rights Reserved.
-#if 0
+
 #pragma once
 
 #include "Containers/ArrayView.h"
@@ -39,4 +39,3 @@ PLAINPROPS_API void LoadRange(uint8* Dst, FRangeView Src, ERangeSizeType MaxType
 
 } // namespace PlainProps
 
-#endif
\ No newline at end of file
diff --git a/Engine/Plugins/Experimental/PlainProps/Source/Public/PlainPropsRead.h b/Engine/Plugins/Experimental/PlainProps/Source/Public/PlainPropsRead.h
--- a/Engine/Plugins/Experimental/PlainProps/Source/Public/PlainPropsRead.h
+++ b/Engine/Plugins/Experimental/PlainProps/Source/Public/PlainPropsRead.h
@@ -1,5 +1,5 @@
 // Copyright Epic Games, Inc. All Rights Reserved.
-#if 0
+
 #pragma once
 
 #include "Containers/ArrayView.h"
@@ -616,5 +616,4 @@ private:
 	}
 };
 
-} // namespace PlainProps
-#endif
\ No newline at end of file
+} // namespace PlainProps
\ No newline at end of file
diff --git a/Engine/Plugins/Experimental/PlainProps/Source/Public/PlainPropsSave.h b/Engine/Plugins/Experimental/PlainProps/Source/Public/PlainPropsSave.h
--- a/Engine/Plugins/Experimental/PlainProps/Source/Public/PlainPropsSave.h
+++ b/Engine/Plugins/Experimental/PlainProps/Source/Public/PlainPropsSave.h
@@ -1,5 +1,5 @@
 // Copyright Epic Games, Inc. All Rights Reserved.
-#if 0
+
 #pragma once
 
 #include "PlainPropsTypes.h"
@@ -25,4 +25,3 @@ struct FSaveContext
 
 } // namespace PlainProps
 
-#endif
\ No newline at end of file
diff --git a/Engine/Plugins/Experimental/PlainProps/Source/Public/PlainPropsTypes.h b/Engine/Plugins/Experimental/PlainProps/Source/Public/PlainPropsTypes.h
--- a/Engine/Plugins/Experimental/PlainProps/Source/Public/PlainPropsTypes.h
+++ b/Engine/Plugins/Experimental/PlainProps/Source/Public/PlainPropsTypes.h
@@ -1,5 +1,5 @@
 // Copyright Epic Games, Inc. All Rights Reserved.
-#if 0
+
 #pragma once
 
 #include "Containers/ArrayView.h"
@@ -348,6 +348,18 @@ template<> inline constexpr FUnpackedLeafType ReflectLeaf<char32_t>	= { ELeafTyp
 
 //////////////////////////////////////////////////////////////////////////
 
+inline constexpr ERangeSizeType RangeSizeOf(bool)	{ return ERangeSizeType::Uni; }
+inline constexpr ERangeSizeType RangeSizeOf(int8)	{ return ERangeSizeType::S8; }
+inline constexpr ERangeSizeType RangeSizeOf(int16)	{ return ERangeSizeType::S16; }
+inline constexpr ERangeSizeType RangeSizeOf(int32)	{ return ERangeSizeType::S32; }
+inline constexpr ERangeSizeType RangeSizeOf(int64)	{ return ERangeSizeType::S64; }
+inline constexpr ERangeSizeType RangeSizeOf(uint8)	{ return ERangeSizeType::U8; }
+inline constexpr ERangeSizeType RangeSizeOf(uint16)	{ return ERangeSizeType::U16; }
+inline constexpr ERangeSizeType RangeSizeOf(uint32)	{ return ERangeSizeType::U32; }
+inline constexpr ERangeSizeType RangeSizeOf(uint64)	{ return ERangeSizeType::U64; }
+
+//////////////////////////////////////////////////////////////////////////
+
 template<typename T>
 const T* AlignPtr(const void* Ptr)
 {
@@ -381,5 +393,4 @@ public:
 	PLAINPROPS_API FString				Print(FStructSchemaId Name) const;
 };
 
-} // namespace PlainProps
-#endif
\ No newline at end of file
+} // namespace PlainProps
\ No newline at end of file
diff --git a/Engine/Plugins/Experimental/PlainProps/Source/Public/PlainPropsUeCoreBindings.h b/Engine/Plugins/Experimental/PlainProps/Source/Public/PlainPropsUeCoreBindings.h
--- a/Engine/Plugins/Experimental/PlainProps/Source/Public/PlainPropsUeCoreBindings.h
+++ b/Engine/Plugins/Experimental/PlainProps/Source/Public/PlainPropsUeCoreBindings.h
@@ -1,5 +1,5 @@
 // Copyright Epic Games, Inc. All Rights Reserved.
-#if 0
+
 #pragma once
 
 #include "Containers/Set.h"
@@ -590,5 +590,4 @@ namespace PlainProps
 	{
 		using Type = UE::TSetBinding<T>;
 	};
-}
-#endif
\ No newline at end of file
+}
\ No newline at end of file
diff --git a/Engine/Plugins/Experimental/PlainProps/Source/Public/PlainPropsWrite.h b/Engine/Plugins/Experimental/PlainProps/Source/Public/PlainPropsWrite.h
--- a/Engine/Plugins/Experimental/PlainProps/Source/Public/PlainPropsWrite.h
+++ b/Engine/Plugins/Experimental/PlainProps/Source/Public/PlainPropsWrite.h
@@ -1,5 +1,5 @@
 // Copyright Epic Games, Inc. All Rights Reserved.
-#if 0
+
 #pragma once
 
 #include "PlainPropsTypes.h"
@@ -36,5 +36,4 @@ private:
 	TUniquePtr<FWriteIds>							NewIds;
 };
 
-} // namespace PlainProps
-#endif
\ No newline at end of file
+} // namespace PlainProps
\ No newline at end of file
diff --git a/Engine/Plugins/Experimental/PlainProps/Source/Tests/PlainPropsCttiTests.cpp b/Engine/Plugins/Experimental/PlainProps/Source/Tests/PlainPropsCttiTests.cpp
--- a/Engine/Plugins/Experimental/PlainProps/Source/Tests/PlainPropsCttiTests.cpp
+++ b/Engine/Plugins/Experimental/PlainProps/Source/Tests/PlainPropsCttiTests.cpp
@@ -1,5 +1,5 @@
 // Copyright Epic Games, Inc. All Rights Reserved.
-#if 0
+
 #include "PlainPropsCtti.h"
 #include <type_traits>
 #include <string_view>
@@ -123,5 +123,4 @@ static_assert(CttiOf<S2<int>>::NumVars == 1);
 static_assert(CttiOf<S2<int>>::Var<0>::Name == std::string_view("a"));
 static_assert(CttiOf<S2<int>>::Var<0>::Offset == offsetof(S2<int>, a));
 
-}
-#endif
\ No newline at end of file
+}
\ No newline at end of file
diff --git a/Engine/Plugins/Experimental/PlainProps/Source/Tests/PlainPropsTests.cpp b/Engine/Plugins/Experimental/PlainProps/Source/Tests/PlainPropsTests.cpp
--- a/Engine/Plugins/Experimental/PlainProps/Source/Tests/PlainPropsTests.cpp
+++ b/Engine/Plugins/Experimental/PlainProps/Source/Tests/PlainPropsTests.cpp
@@ -1,5 +1,5 @@
 // Copyright Epic Games, Inc. All Rights Reserved.
-#if 0
+
 #if WITH_TESTS
 #include "Tests/TestHarnessAdapter.h"
 #include "Containers/AnsiString.h"
@@ -1350,5 +1350,4 @@ TEST_CASE_NAMED(FPlainPropsLoadSaveTest, "System::Core::Serialization::PlainProp
 }
 
 } // namespace PlainProps
-#endif // WITH_TESTS
-#endif
\ No newline at end of file
+#endif // WITH_TESTS
\ No newline at end of file
diff --git a/Engine/Plugins/Experimental/PlainProps/Source/Tests/PlainPropsUeCoreTests.cpp b/Engine/Plugins/Experimental/PlainProps/Source/Tests/PlainPropsUeCoreTests.cpp
--- a/Engine/Plugins/Experimental/PlainProps/Source/Tests/PlainPropsUeCoreTests.cpp
+++ b/Engine/Plugins/Experimental/PlainProps/Source/Tests/PlainPropsUeCoreTests.cpp
@@ -1,5 +1,5 @@
 // Copyright Epic Games, Inc. All Rights Reserved.
-#if 0
+
 #if WITH_TESTS
 
 #include "PlainPropsBuildSchema.h"
@@ -761,5 +761,4 @@ TEST_CASE_NAMED(FPlainPropsUeCoreTest, "System::Core::Serialization::PlainProps:
 }
 
 } // namespace PlainProps::UE::Test
-#endif // WITH_TESTS
-#endif
\ No newline at end of file
+#endif // WITH_TESTS
\ No newline at end of file
