PlainProps add module

#rnx

diff --git a/Engine/Plugins/Experimental/PlainProps/Source/Private/PlainPropsModule.cpp b/Engine/Plugins/Experimental/PlainProps/Source/Private/PlainPropsModule.cpp
new file mode 100644
--- /dev/null
+++ b/Engine/Plugins/Experimental/PlainProps/Source/Private/PlainPropsModule.cpp
@@ -0,0 +1,6 @@
+// Copyright Epic Games, Inc. All Rights Reserved.
+
+#include "PlainPropsModule.h"
+#include "Modules/ModuleManager.h"
+
+IMPLEMENT_MODULE(FPlainPropsModule, PlainProps);
\ No newline at end of file
diff --git a/Engine/Plugins/Experimental/PlainProps/Source/Public/PlainPropsModule.h b/Engine/Plugins/Experimental/PlainProps/Source/Public/PlainPropsModule.h
new file mode 100644
--- /dev/null
+++ b/Engine/Plugins/Experimental/PlainProps/Source/Public/PlainPropsModule.h
@@ -0,0 +1,7 @@
+// Copyright Epic Games, Inc. All Rights Reserved.
+
+#pragma once
+
+#include "Modules/ModuleInterface.h"
+
+class FPlainPropsModule final : public IModuleInterface {};
\ No newline at end of file
