Properly manage allowed platforms for a plugin
#rnx

diff --git a/Engine/Plugins/Experimental/PlainProps/PlainProps.uplugin b/Engine/Plugins/Experimental/PlainProps/PlainProps.uplugin
--- a/Engine/Plugins/Experimental/PlainProps/PlainProps.uplugin
+++ b/Engine/Plugins/Experimental/PlainProps/PlainProps.uplugin
@@ -7,10 +7,14 @@
 	"CreatedByURL" : "https://epicgames.com",
 	"EnabledByDefault" : false,
 	"IsExperimentalVersion": true,
+	"SupportedTargetPlatforms": [
+		"Win64"
+	],
 	"Modules": [
 		{
 			"Name" : "PlainProps",
-			"Type" : "Runtime"
+			"Type" : "Runtime",
+			"PlatformAllowList": [ "Win64" ]
 		}
 	]
 }
\ No newline at end of file
diff --git a/Engine/Plugins/Experimental/PlainProps/Source/PlainProps.Build.cs b/Engine/Plugins/Experimental/PlainProps/Source/PlainProps.Build.cs
--- a/Engine/Plugins/Experimental/PlainProps/Source/PlainProps.Build.cs
+++ b/Engine/Plugins/Experimental/PlainProps/Source/PlainProps.Build.cs
@@ -2,7 +2,6 @@
 
 using UnrealBuildTool;
 
-[SupportedPlatforms("Win64", "PS5")]
 public class PlainProps : ModuleRules
 {
 	public PlainProps(ReadOnlyTargetRules Target) : base(Target)
