PlainProps changes

+ Transform custom binding
+ FCustomBindings layering
+ Struct range delta serialization / struct default lookup mechanism
+ PP_NAME_STRUCT
+ PLAINPROPS_API sprinkling
+ FBuiltStructPtr/Deleter and FLoadBatchPtr/Deleter
+ ESchemaFormat::InMemoryNames
* Fix automatic binding of native custom struct members
* Fix int64 member bindings on clang
* Change to FPlatformMemory::KernelAddressBit

#rnx

diff --git a/Engine/Plugins/Experimental/PlainProps/Source/Private/PlainPropsBind.cpp b/Engine/Plugins/Experimental/PlainProps/Source/Private/PlainPropsBind.cpp
--- a/Engine/Plugins/Experimental/PlainProps/Source/Private/PlainPropsBind.cpp
+++ b/Engine/Plugins/Experimental/PlainProps/Source/Private/PlainPropsBind.cpp
@@ -177,7 +177,7 @@ ICustomBinding*	FCustomBindings::Find(FStructSchemaId Id) const
 		}
 	}
 
-	return nullptr;
+	return Base ? Base->Find(Id) : nullptr;
 }
 
 ////////////////////////////////////////////////////////////////////////////////////////////////
@@ -254,6 +254,22 @@ void FSchemaBindings::DropStruct(FStructSchemaId Id)
 
 //////////////////////////////////////////////////////////////////////////
 
+TArray<FStructSchemaId> IndexInMemoryNames(const FSchemaBatch& Schemas, FIdIndexerBase& Indexer)
+{
+	const uint8* Base = reinterpret_cast<const uint8*>(&Schemas);
+	const uint32* Offsets = Schemas.SchemaOffsets;
+
+	TArray<FStructSchemaId> Out;
+	Out.SetNumUninitialized(Schemas.NumStructSchemas);
+	for (int32 Idx = 0; Idx < Out.Num(); ++Idx)
+	{
+		const FStructSchema& Schema = *reinterpret_cast<const FStructSchema*>(Base + Offsets[Idx]);
+		Out[Idx] = Indexer.IndexStruct(Schema.Type);
+	}
+
+	return Out;
+}
+
 uint32 FIdTranslatorBase::CalculateTranslationSize(int32 NumSavedNames, const FSchemaBatch& Batch)
 {
 	static_assert(sizeof(FNameId) == sizeof(FNestedScopeId));
@@ -355,7 +371,7 @@ FSchemaBatch* CreateTranslatedSchemas(const FSchemaBatch& In, FIdBinding NewIds)
 		Out->SchemaOffsets[Idx] = In.SchemaOffsets[Idx] - DroppedBytes;
 	}
 
-	// Copy schemas and remap type ids
+	// Copy schemas and remap type ids if needed
 	FMemory::Memcpy(reinterpret_cast<uint8*>(Out) + Out->GetSchemaOffsets()[0], InSchemas.GetData(), InSchemas.GetSize());
 	for (FStructSchema& Schema : GetStructSchemas(*Out))
 	{
diff --git a/Engine/Plugins/Experimental/PlainProps/Source/Private/PlainPropsBuild.cpp b/Engine/Plugins/Experimental/PlainProps/Source/Private/PlainPropsBuild.cpp
--- a/Engine/Plugins/Experimental/PlainProps/Source/Private/PlainPropsBuild.cpp
+++ b/Engine/Plugins/Experimental/PlainProps/Source/Private/PlainPropsBuild.cpp
@@ -11,6 +11,11 @@
 namespace PlainProps
 {
 
+void FBuiltStructDeleter::operator()(FBuiltStruct* Ptr) const
+{
+	delete Ptr;
+}
+
 FBuiltRange* FBuiltRange::Create(uint64 NumItems, SIZE_T ItemSize)
 {
 	check(NumItems > 0);
@@ -53,7 +58,7 @@ namespace Private
 		return PlainProps::BuildStructuralRangeImpl(Values);
 	}
 
-	FBuiltRange* BuildStructuralRange(TArrayView64<TUniquePtr<FBuiltStruct>> Values)
+	FBuiltRange* BuildStructuralRange(TArrayView64<FBuiltStructPtr> Values)
 	{
 		return PlainProps::BuildStructuralRangeImpl(Values);
 	}
@@ -96,7 +101,7 @@ namespace Private
 FMemberBuilder::FMemberBuilder() {}
 FMemberBuilder::~FMemberBuilder() {}
 
-void FMemberBuilder::AddStruct(FMemberId Name, FStructSchemaId Schema, TUniquePtr<FBuiltStruct>&& Struct)
+void FMemberBuilder::AddStruct(FMemberId Name, FStructSchemaId Schema, FBuiltStructPtr&& Struct)
 {
 	Members.Emplace(Name, Schema, MoveTemp(Struct));
 }
@@ -123,12 +128,12 @@ void FMemberBuilder::BuildSuperStruct(const FStructDeclaration& Super, const FDe
 		return;
 	}
 	
-	TUniquePtr<FBuiltStruct> OnlyMember = BuildAndReset(Super, Debug);
+	FBuiltStructPtr OnlyMember = BuildAndReset(Super, Debug);
 	Members.Emplace(FBuiltMember::MakeSuper(Super.Id, MoveTemp(OnlyMember)));
 	check(Members[0].Schema.Type.AsStruct().IsSuper);
 }
 
-TUniquePtr<FBuiltStruct> FMemberBuilder::BuildAndReset(const FStructDeclaration& Declared, const FDebugIds& Debug)
+FBuiltStructPtr FMemberBuilder::BuildAndReset(const FStructDeclaration& Declared, const FDebugIds& Debug)
 {
 	checkf(!(Declared.Super && Declared.Occupancy == EMemberPresence::RequireAll),
 		TEXT("Requiring sub structs to be dense isn't implemented"));
@@ -161,9 +166,9 @@ TUniquePtr<FBuiltStruct> FMemberBuilder::BuildAndReset(const FStructDeclaration&
 
 	Members.Reset();
 
-	GLiveStructsFoo.Add(Out);
+	//GLiveStructsFoo.Add(Out);
 
-	return TUniquePtr<FBuiltStruct>(Out);
+	return FBuiltStructPtr(Out);
 }
 
 template<typename IntType, typename FloatType>
@@ -209,11 +214,11 @@ FBuiltMember::FBuiltMember(FMemberId Name, FTypedRange&& Range)
 : FBuiltMember(Name, MoveTemp(Range.Schema), { .Range = Range.Values })
 {}
 
-FBuiltMember::FBuiltMember(FMemberId Name, FStructSchemaId Schema, TUniquePtr<FBuiltStruct>&& Value)
+FBuiltMember::FBuiltMember(FMemberId Name, FStructSchemaId Schema, FBuiltStructPtr&& Value)
 : FBuiltMember(Name, {DefaultStructType, FOptionalSchemaId(Schema)}, { .Struct = Value.Release() })
 {}
 
-FBuiltMember FBuiltMember::MakeSuper(FStructSchemaId Schema, TUniquePtr<FBuiltStruct>&& Value)
+FBuiltMember FBuiltMember::MakeSuper(FStructSchemaId Schema, FBuiltStructPtr&& Value)
 {
 	return FBuiltMember(NoId, {SuperStructType, FOptionalSchemaId(Schema)}, { .Struct = Value.Release() });
 }
@@ -266,7 +271,7 @@ uint64 FBuiltRange::Delete(FBuiltRange* Range, FOptionalSchemaId InnerSchema, TC
 
 FBuiltStruct::~FBuiltStruct()
 { 
-	check(GLiveStructsFoo.Remove(this) == 1);
+	//check(GLiveStructsFoo.Remove(this) == 1);
 
 	for (const FBuiltMember& Member : MakeArrayView(Members, NumMembers))
 	{
@@ -278,7 +283,7 @@ FBuiltStruct::~FBuiltStruct()
 
 FTypedRange FStructRangeBuilder::BuildAndReset(const FStructDeclaration& Declared, const FDebugIds& Debug)
 {
-	TArray64<TUniquePtr<FBuiltStruct>> BuiltStructs;
+	TArray64<FBuiltStructPtr> BuiltStructs;
 	BuiltStructs.Reserve(Structs.Num());
 	for (FMemberBuilder& Struct : Structs)
 	{
diff --git a/Engine/Plugins/Experimental/PlainProps/Source/Private/PlainPropsBuildSchema.cpp b/Engine/Plugins/Experimental/PlainProps/Source/Private/PlainPropsBuildSchema.cpp
--- a/Engine/Plugins/Experimental/PlainProps/Source/Private/PlainPropsBuildSchema.cpp
+++ b/Engine/Plugins/Experimental/PlainProps/Source/Private/PlainPropsBuildSchema.cpp
@@ -69,6 +69,10 @@ struct FEnumSchemaBuilder
 
 //////////////////////////////////////////////////////////////////////////
 
+FSchemasBuilder::FSchemasBuilder(const FDeclarations& Types)
+: FSchemasBuilder(Types.GetStructs(), Types.GetEnums(), Types.GetDebug())
+{}
+
 FSchemasBuilder::FSchemasBuilder(FStructDeclarations InStructs, FEnumDeclarations InEnums, const FDebugIds& InDebug)
 : DeclaredStructs(InStructs)
 , DeclaredEnums(InEnums)
diff --git a/Engine/Plugins/Experimental/PlainProps/Source/Private/PlainPropsInternalBuild.h b/Engine/Plugins/Experimental/PlainProps/Source/Private/PlainPropsInternalBuild.h
--- a/Engine/Plugins/Experimental/PlainProps/Source/Private/PlainPropsInternalBuild.h
+++ b/Engine/Plugins/Experimental/PlainProps/Source/Private/PlainPropsInternalBuild.h
@@ -24,8 +24,8 @@ struct FBuiltMember
 	FBuiltMember(FMemberId Name, FUnpackedLeafType Leaf, FOptionalEnumSchemaId Schema, uint64 Value);
 	//FBuiltMember(FMemberId Name, FEnumSchemaId Schema, ELeafWidth Width,  uint64 Value);
 	FBuiltMember(FMemberId Name, FTypedRange&& Range);
-	FBuiltMember(FMemberId Name, FStructSchemaId Schema, TUniquePtr<FBuiltStruct>&& Value);
-	static FBuiltMember MakeSuper(FStructSchemaId Schema, TUniquePtr<FBuiltStruct>&& Value);
+	FBuiltMember(FMemberId Name, FStructSchemaId Schema, FBuiltStructPtr&& Value);
+	static FBuiltMember MakeSuper(FStructSchemaId Schema, FBuiltStructPtr&& Value);
 	~FBuiltMember(); // Deletes Value
 
 	FBuiltMember& operator=(const FBuiltMember&) = delete;
diff --git a/Engine/Plugins/Experimental/PlainProps/Source/Private/PlainPropsLoad.cpp b/Engine/Plugins/Experimental/PlainProps/Source/Private/PlainPropsLoad.cpp
--- a/Engine/Plugins/Experimental/PlainProps/Source/Private/PlainPropsLoad.cpp
+++ b/Engine/Plugins/Experimental/PlainProps/Source/Private/PlainPropsLoad.cpp
@@ -36,7 +36,7 @@ public:
 
 	// @param OffsetWidth Usage unimplemented, store size and offsets as 8/16/32/64-bit 
 	explicit FLoadStructPlan(const FSchemaBinding& Schema, ELeafWidth OffsetWidth, bool bSparse)
-	: Handle(uint64(&Schema) | (uint64(OffsetWidth) << 1) | (uint64(bSparse) << KernelAddressBit) | SchemaMask)
+	: Handle(uint64(&Schema) | (uint64(OffsetWidth) << 1) | (uint64(bSparse) << FPlatformMemory::KernelAddressBit) | SchemaMask)
 	{
 		check(&Schema == &AsSchema());
 		check(IsSparseSchema() == bSparse);
@@ -51,7 +51,7 @@ public:
 	const FSchemaBinding&		AsSchema() const		{ check(IsSchema()); return *AsPtr<FSchemaBinding>(); }
 
 private:
-	static constexpr uint64 SparseMask			= uint64(1) << KernelAddressBit;
+	static constexpr uint64 SparseMask			= uint64(1) << FPlatformMemory::KernelAddressBit;
 	static constexpr uint64 PtrMask				= ~(SparseMask | 0b111);
 	static constexpr uint64 LoMask				= 0b11;
 	static constexpr uint64 MemcpyMask			= 0b00;
@@ -133,6 +133,11 @@ struct FLoadBatch
 	FLoadStructPlan			operator[](FStructSchemaId Id) const { check(Id.Idx < NumPlans); return Plans[Id.Idx]; }
 };
 
+void FLoadBatchDeleter::operator()(FLoadBatch* Batch) const
+{
+	FMemory::Free(Batch);
+}
+
 using SubsetByteArray = TArray<uint8, TInlineAllocator<1024>>;
 
 static void CopyMemberBinding(FLeafMemberBinding Binding, /* in-out */ const FSchemaId*& InnerSchemaIt, FMemberBinder& Out)
@@ -244,7 +249,7 @@ static void CloneBindingWithReplacedStructIds(const FSchemaId* FromIds, const FS
 	return Memcpy ? FLoadStructPlan(Memcpy.GetValue()) : MakeSchemaLoadPlan(From, To, ToMemberIds, ToStructIds, OutSubsetSchemas);
 }
 
-FLoadBatch* CreateLoadPlans(FReadBatchId ReadId, const FDeclarations& Declarations, const FCustomBindings& Customs, const FSchemaBindings& Schemas, TConstArrayView<FStructSchemaId> RuntimeIds)
+FLoadBatchPtr CreateLoadPlans(FReadBatchId ReadId, const FDeclarations& Declarations, const FCustomBindings& Customs, const FSchemaBindings& Schemas, TConstArrayView<FStructSchemaId> RuntimeIds)
 {
 	check(NumStructSchemas(ReadId) == RuntimeIds.Num());
 
@@ -303,12 +308,7 @@ FLoadBatch* CreateLoadPlans(FReadBatchId ReadId, const FDeclarations& Declaratio
 		check(It == OutSubsetData + SubsetSchemaData.Num());
 	}
 
-	return Out;
-}
-
-void DestroyLoadPlans(FLoadBatch* Batch)
-{
-	FMemory::Free(Batch);
+	return FLoadBatchPtr(Out);
 }
 
 ////////////////////////////////////////////////////////////////////////////
@@ -491,7 +491,7 @@ public:
 	, ByteIt(Values)
 	{}
 
-	void Load(uint8* Struct)
+	void Load(void* Struct)
 	{
 		SkipMissingSparseMembers();
 	
@@ -532,10 +532,10 @@ private:
 		}
 	}
 
-	void LoadMember(uint8* Struct)
+	void LoadMember(void* Struct)
 	{
 		FMemberBindType Type = Types[MemberIdx];
-		uint8* Member = Struct + Offsets[MemberIdx];
+		uint8* Member = static_cast<uint8*>(Struct) + Offsets[MemberIdx];
 
 		switch (Type.GetKind())
 		{
@@ -593,7 +593,7 @@ private:
 
 ////////////////////////////////////////////////////////////////////////////
 
-void LoadStruct(uint8* Dst, FByteReader Src, FStructSchemaId Id, const FLoadBatch& Batch)
+void LoadStruct(void* Dst, FByteReader Src, FStructSchemaId Id, const FLoadBatch& Batch)
 {
 	FLoadStructPlan Plan = Batch[Id];
 	if (Plan.IsSchema())
@@ -610,7 +610,7 @@ void LoadStruct(uint8* Dst, FByteReader Src, FStructSchemaId Id, const FLoadBatc
 	else if (Plan.IsMemcpy())
 	{
 		Src.CheckSize(Plan.AsMemcpy().Size);
-		FMemory::Memcpy(Dst + Plan.AsMemcpy().Offset, Src.Peek(), Plan.AsMemcpy().Size);
+		FMemory::Memcpy(static_cast<uint8*>(Dst) + Plan.AsMemcpy().Offset, Src.Peek(), Plan.AsMemcpy().Size);
 	}
 	else
 	{
@@ -619,12 +619,12 @@ void LoadStruct(uint8* Dst, FByteReader Src, FStructSchemaId Id, const FLoadBatc
 	}
 }
 
-void LoadStruct(uint8* Dst, FStructView Src, const FLoadBatch& Batch)
+void LoadStruct(void* Dst, FStructView Src, const FLoadBatch& Batch)
 {
 	LoadStruct(Dst, Src.Values, Src.Schema.Id, Batch);
 }
 
-void ConstructAndLoadStruct(uint8* Dst, FByteReader Src, FStructSchemaId Id, const FLoadBatch& Batch)
+void ConstructAndLoadStruct(void* Dst, FByteReader Src, FStructSchemaId Id, const FLoadBatch& Batch)
 {
 	FLoadStructPlan Plan = Batch[Id];
 	checkf(!Plan.IsSchema(), TEXT("Non-default constructible types requires ICustomBinding or in rare cases memcpying"));
@@ -632,7 +632,7 @@ void ConstructAndLoadStruct(uint8* Dst, FByteReader Src, FStructSchemaId Id, con
 	if (Plan.IsMemcpy())
 	{
 		Src.CheckSize(Plan.AsMemcpy().Size);
-		FMemory::Memcpy(Dst + Plan.AsMemcpy().Offset, Src.Peek(), Plan.AsMemcpy().Size);
+		FMemory::Memcpy(static_cast<uint8*>(Dst) + Plan.AsMemcpy().Offset, Src.Peek(), Plan.AsMemcpy().Size);
 	}
 	else
 	{
@@ -641,9 +641,9 @@ void ConstructAndLoadStruct(uint8* Dst, FByteReader Src, FStructSchemaId Id, con
 	}
 }
 
-void LoadRange(uint8* Dst, FRangeView Src, ERangeSizeType MaxSize, TConstArrayView<FRangeBinding> Bindings, const FLoadBatch& Batch)
+void LoadRange(void* Dst, FRangeView Src, ERangeSizeType MaxSize, TConstArrayView<FRangeBinding> Bindings, const FLoadBatch& Batch)
 {
-	FRangeLoader::LoadRangeView(Dst, Src, MaxSize, Bindings, Batch);
+	FRangeLoader::LoadRangeView(static_cast<uint8*>(Dst), Src, MaxSize, Bindings, Batch);
 }
 
 
diff --git a/Engine/Plugins/Experimental/PlainProps/Source/Private/PlainPropsSave.cpp b/Engine/Plugins/Experimental/PlainProps/Source/Private/PlainPropsSave.cpp
--- a/Engine/Plugins/Experimental/PlainProps/Source/Private/PlainPropsSave.cpp
+++ b/Engine/Plugins/Experimental/PlainProps/Source/Private/PlainPropsSave.cpp
@@ -4,6 +4,7 @@
 #include "PlainPropsBind.h"
 #include "PlainPropsInternalBuild.h"
 #include "PlainPropsInternalFormat.h"
+#include "Algo/Find.h"
 #include <type_traits>
 
 namespace PlainProps
@@ -98,7 +99,8 @@ struct TStructuralRangeSaver
 };
 
 using FNestedRangeSaver = TStructuralRangeSaver<FBuiltRange*, FRangeMemberBinding>;
-using FStructRangeSaver = TStructuralRangeSaver<TUniquePtr<FBuiltStruct>, FStructSchemaId>;
+using FStructRangeSaver = TStructuralRangeSaver<FBuiltStructPtr, FStructSchemaId>;
+using FStructRangeDeltaSaver = TStructuralRangeSaver<FBuiltStructPtr, FDefaultStruct>;
 
 //////////////////////////////////////////////////////////////////////////
 
@@ -150,6 +152,18 @@ template<class SaverType, typename InnerContextType>
 	return Allocator.GetAllocatedRange();
 }
 
+[[nodiscard]] static FBuiltRange* SaveStructRange(const void* Range, const IItemRangeBinding& ItemBinding, const FSaveContext& Ctx, FStructSchemaId Id)
+{
+	if (const FDefaultStruct* Defaults = Algo::FindBy(Ctx.Defaults, Id, &FDefaultStruct::Id))
+	{
+		return SaveRange<FStructRangeDeltaSaver>(Range, ItemBinding, Ctx, *Defaults);
+	}
+	else
+	{
+		return SaveRange<FStructRangeSaver>(Range, ItemBinding, Ctx, Id);
+	}
+}
+
 [[nodiscard]] static FBuiltRange* SaveRange(const uint8* Range, FRangeMemberBinding Member, const FSaveContext& Ctx)
 {
 	FRangeBinding Binding = Member.RangeBindings[0];
@@ -163,9 +177,9 @@ template<class SaverType, typename InnerContextType>
 	const IItemRangeBinding& ItemBinding = Binding.AsItemBinding();
 	switch (InnerType.GetKind())
 	{
-	case EMemberKind::Leaf:		return SaveRange<FLeafRangeSaver>(  Range, ItemBinding, Ctx, SizeOf(GetArithmeticWidth(InnerType.AsLeaf())));
+	case EMemberKind::Leaf:		return SaveRange<FLeafRangeSaver>(Range, ItemBinding, Ctx, SizeOf(GetArithmeticWidth(InnerType.AsLeaf())));
 	case EMemberKind::Range:	return SaveRange<FNestedRangeSaver>(Range, ItemBinding, Ctx, GetInnerRange(Member));
-	case EMemberKind::Struct:	return SaveRange<FStructRangeSaver>(Range, ItemBinding, Ctx, static_cast<FStructSchemaId>(Member.InnermostSchema.Get()));
+	case EMemberKind::Struct:	return SaveStructRange(Range, ItemBinding, Ctx, static_cast<FStructSchemaId>(Member.InnermostSchema.Get()));
 	}
 
 	check(false);
@@ -179,11 +193,16 @@ template<class SaverType, typename InnerContextType>
 	return SaveRange(Range, Member, Ctx);
 }
 
-[[nodiscard]] static TUniquePtr<FBuiltStruct> SaveRangeItem(const uint8* Struct, FStructSchemaId Id, const FSaveContext& Ctx)
+[[nodiscard]] static FBuiltStructPtr SaveRangeItem(const uint8* Struct, FStructSchemaId Id, const FSaveContext& Ctx)
 {
 	return SaveStruct(Struct, Id, Ctx);
 }
 
+[[nodiscard]] static FBuiltStructPtr SaveRangeItem(const uint8* Struct, FDefaultStruct Default, const FSaveContext& Ctx)
+{
+	return SaveStructDelta(Struct, Default.Struct, Default.Id, Ctx);
+}
+
 [[nodiscard]] static FMemberType ToMemberType(FMemberBindType In)
 {
 	switch (In.GetKind())
@@ -233,7 +252,7 @@ static void SaveMember(FMemberBuilder& Out, const void* Struct, FMemberId Name,
 	Out.AddStruct(Name, Member.Id, SaveStruct(At(Struct, Member.Offset), Member.Id, Ctx));
 }
 
-TUniquePtr<FBuiltStruct> SaveStruct(const void* Struct, FStructSchemaId Id, const FSaveContext& Ctx)
+FBuiltStructPtr SaveStruct(const void* Struct, FStructSchemaId Id, const FSaveContext& Ctx)
 {
 	const FStructDeclaration& Declaration = Ctx.Declarations.Get(Id);
 
@@ -441,13 +460,13 @@ static void SaveMemberDelta(FMemberBuilder& Out, const void* Struct, const void*
 
 static void SaveMemberDelta(FMemberBuilder& Out, const void* Struct, const void* Default, FMemberId Name, const FSaveContext& Ctx, FStructMemberBinding Member)
 {
-	if (TUniquePtr<FBuiltStruct> Delta = SaveStructDelta(At(Struct, Member.Offset), At(Default, Member.Offset), Member.Id, Ctx))
+	if (FBuiltStructPtr Delta = SaveStructDelta(At(Struct, Member.Offset), At(Default, Member.Offset), Member.Id, Ctx))
 	{
 		Out.AddStruct(Name, Member.Id, MoveTemp(Delta));
 	}
 }
 
-TUniquePtr<FBuiltStruct> SaveStructDelta(const void* Struct, const void* Default, FStructSchemaId Id, const FSaveContext& Ctx)
+FBuiltStructPtr SaveStructDelta(const void* Struct, const void* Default, FStructSchemaId Id, const FSaveContext& Ctx)
 {
 	const FStructDeclaration& Declaration = Ctx.Declarations.Get(Id);
 
diff --git a/Engine/Plugins/Experimental/PlainProps/Source/Private/PlainPropsUeCoreBindings.cpp b/Engine/Plugins/Experimental/PlainProps/Source/Private/PlainPropsUeCoreBindings.cpp
new file mode 100644
--- /dev/null
+++ b/Engine/Plugins/Experimental/PlainProps/Source/Private/PlainPropsUeCoreBindings.cpp
@@ -0,0 +1,129 @@
+// Copyright Epic Games, Inc. All Rights Reserved.
+
+#include "PlainPropsUeCoreBindings.h"
+#include "PlainPropsBuild.h"
+#include "Math/Transform.h"
+
+namespace PlainProps::UE
+{
+
+void FTransformBinding::Save(FMemberBuilder& Dst, const FTransform& Src, const FTransform* Default, const FTransformIds& Ids) const
+{
+	static_assert(std::is_same_v<decltype(FTransform().GetTranslation().X), double>);
+
+	FVector T = Src.GetTranslation();
+	FQuat R = Src.GetRotation();
+	FVector S = Src.GetScale3D();
+
+	if (Default)
+	{
+		if (T != Default->GetTranslation())
+		{
+			Dst.Add(Ids.Translate[0], T.X);
+			Dst.Add(Ids.Translate[1], T.Y);
+			Dst.Add(Ids.Translate[2], T.Z);
+		}
+
+		if (R != Default->GetRotation())
+		{
+			Dst.Add(Ids.Rotate[0], R.X);
+			Dst.Add(Ids.Rotate[1], R.Y);
+			Dst.Add(Ids.Rotate[2], R.Z);
+			Dst.Add(Ids.Rotate[3], R.W);
+		}
+
+		if (S != Default->GetScale3D())
+		{
+			Dst.Add(Ids.Scale[0], S.X);
+			Dst.Add(Ids.Scale[1], S.Y);
+			Dst.Add(Ids.Scale[2], S.Z);
+		}
+	}
+	else
+	{
+		Dst.Add(Ids.Translate[0],	T.X);
+		Dst.Add(Ids.Translate[1],	T.Y);
+		Dst.Add(Ids.Translate[2],	T.Z);
+		Dst.Add(Ids.Rotate[0],		R.X);
+		Dst.Add(Ids.Rotate[1],		R.Y);
+		Dst.Add(Ids.Rotate[2],		R.Z);
+		Dst.Add(Ids.Rotate[3],		R.W);
+		Dst.Add(Ids.Scale[0],		S.X);
+		Dst.Add(Ids.Scale[1],		S.Y);
+		Dst.Add(Ids.Scale[2],		S.Z);
+	}
+}
+
+void FTransformBinding::Load(FTransform& Dst, FStructView Src, ECustomLoadMethod Method, const FLoadBatch& Batch, const FTransformIds& Ids) const
+{
+	FMemberReader Members(Src);
+
+	if (Method == ECustomLoadMethod::Construct)
+	{
+		::new (&Dst) FTransform;
+	}
+				
+	if (!Members.HasMore())
+	{
+		return;
+	}
+
+	FMemberId Name = Members.PeekName().Get();
+	double X = Members.GrabLeaf().AsDouble();
+	double Y = Members.GrabLeaf().AsDouble();
+	double Z = Members.GrabLeaf().AsDouble();
+
+	if (Name == Ids.Translate[0])
+	{
+		Dst.SetTranslation({X, Y, Z});
+			
+		if (!Members.HasMore())
+		{
+			return;
+		}
+			
+		Name = Members.PeekName().Get();
+		X = Members.GrabLeaf().AsDouble();
+		Y = Members.GrabLeaf().AsDouble();
+		Z = Members.GrabLeaf().AsDouble();
+	}
+
+	if (Name == Ids.Rotate[0])
+	{
+		double W = Members.GrabLeaf().AsDouble();
+		Dst.SetRotation({X, Y, Z, W});
+
+		if (!Members.HasMore())
+		{
+			return;
+		}
+
+		Name = Members.PeekName().Get();
+		X = Members.GrabLeaf().AsDouble();
+		Y = Members.GrabLeaf().AsDouble();
+		Z = Members.GrabLeaf().AsDouble();
+	}
+
+	check(Name == Ids.Scale[0]);
+	Dst.SetScale3D({X, Y, Z});
+	check(!Members.HasMore());
+}
+
+bool FTransformBinding::DiffStruct(const void* StructA, const void* StructB) const
+{
+	return !static_cast<const FTransform*>(StructA)->Equals(*static_cast<const FTransform*>(StructB), 0.0);
+}
+
+} // namespace PlainProps::UE
+
+//////////////////////////////////////////////////////////////////////////
+namespace PlainProps
+{
+	
+template<>
+void AppendString(FString& Out, const FName& Name)
+{
+	Name.AppendString(Out);
+}
+
+}
\ No newline at end of file
diff --git a/Engine/Plugins/Experimental/PlainProps/Source/Private/PlainPropsWrite.cpp b/Engine/Plugins/Experimental/PlainProps/Source/Private/PlainPropsWrite.cpp
--- a/Engine/Plugins/Experimental/PlainProps/Source/Private/PlainPropsWrite.cpp
+++ b/Engine/Plugins/Experimental/PlainProps/Source/Private/PlainPropsWrite.cpp
@@ -18,13 +18,15 @@ namespace PlainProps
 // Rewrite as a more compact data structure once we get a large number of ids
 struct FWriteIds
 {
-	FWriteIds(const FIdIndexerBase& DeclaredIds, const FBuiltSchemas& Schemas);
+	FWriteIds(const FIdIndexerBase& DeclaredIds, const FBuiltSchemas& Schemas, ESchemaFormat Format);
+
+	bool								HasStableNames() const { return !Names.IsEmpty(); }
 
 	TArray<FOptionalNameId>				Names;
 	TArray<FOptionalNestedScopeId>		NestedScopes;
 	TArray<FOptionalParametricTypeId>	ParametricTypes;
-	TArray<FOptionalSchemaId>			Structs;
-	TArray<FOptionalSchemaId>			Enums;
+	TArray<FOptionalStructSchemaId>		Structs;
+	TArray<FOptionalEnumSchemaId>		Enums;
 
 	uint32								NumKeptSchemas;
 	uint32								NumKeptStructSchemas;
@@ -41,8 +43,8 @@ struct FWriteIds
 	FParametricTypeId					Remap(FParametricTypeId Old) const		{ return ParametricTypes[Old.Idx].Get(); }
 	FTypenameId							Remap(FTypenameId Old) const			{ return Old.IsConcrete() ? FTypenameId(Remap(Old.AsConcrete())) : FTypenameId(Remap(Old.AsParametric())); }
 	FTypeId								Remap(FTypeId Old) const				{ return { Remap(Old.Scope), Remap(Old.Name) }; }
-	FSchemaId							RemapStruct(FSchemaId Old) const		{ return Structs[Old.Idx].Get(); }
-	FSchemaId							RemapEnum(FSchemaId Old) const			{ return Enums[Old.Idx].Get(); }
+	FStructSchemaId						RemapStruct(FSchemaId Old) const		{ return Structs[Old.Idx].Get(); }
+	FEnumSchemaId						RemapEnum(FSchemaId Old) const			{ return Enums[Old.Idx].Get(); }
 };
 
 static TConstArrayView<FNameId> GetUsedNames(const FBuiltStructSchema& Used)
@@ -176,7 +178,7 @@ static void CopyUsedIds(TArray<T>& Out, const TBitArray<>& Used, C&& Ids)
 	}
 }
 
-FWriteIds::FWriteIds(const FIdIndexerBase& Declared, const FBuiltSchemas& Schemas)
+FWriteIds::FWriteIds(const FIdIndexerBase& Declared, const FBuiltSchemas& Schemas, ESchemaFormat Format)
 {
 	Structs.Init(NoId, Declared.NumStructs());
 	Enums.Init(NoId, Declared.NumEnums());
@@ -185,50 +187,53 @@ FWriteIds::FWriteIds(const FIdIndexerBase& Declared, const FBuiltSchemas& Schema
 	uint32 NewSchemaIdx = 0;
 	for (const FBuiltStructSchema& Struct : Schemas.Structs)
 	{
-		Structs[Struct.Id.Idx] = FSchemaId{NewSchemaIdx++};
+		Structs[Struct.Id.Idx] = FStructSchemaId{NewSchemaIdx++};
 	}
 	NumKeptStructSchemas = NewSchemaIdx;
 	for (const FBuiltEnumSchema& Enum : Schemas.Enums)
 	{
-		Enums[Enum.Id.Idx] = FSchemaId{NewSchemaIdx++};
+		Enums[Enum.Id.Idx] = FEnumSchemaId{NewSchemaIdx++};
 	}
 	NumKeptSchemas = NewSchemaIdx;
 
-	// Generate new name indices
-	FUsedIds Used(Declared);
-	Used.DetectUsage(Schemas.Structs);
-	Used.DetectUsage(Schemas.Enums);
-	MakeRemapping(Names, Used.Names);
+	if (Format == ESchemaFormat::StableNames)
+	{
+		// Generate new name indices
+		FUsedIds Used(Declared);
+		Used.DetectUsage(Schemas.Structs);
+		Used.DetectUsage(Schemas.Enums);
+		MakeRemapping(Names, Used.Names);
 
-	// Generate new nested scope and parametric type mappings
-	KeptScopes.Reserve(MakeRemapping(/* out */ NestedScopes, Used.NestedScopes));
-	KeptParametrics.Reserve(MakeRemapping(/* out */ ParametricTypes, Used.ParametricTypes, Declared.GetParametricTypes()));
+		// Generate new nested scope and parametric type mappings
+		KeptScopes.Reserve(MakeRemapping(/* out */ NestedScopes, Used.NestedScopes));
+		KeptParametrics.Reserve(MakeRemapping(/* out */ ParametricTypes, Used.ParametricTypes, Declared.GetParametricTypes()));
 
-	// Copy nested scopes and parametric types
-	CopyUsedIds(/* out */ KeptScopes, Used.NestedScopes, Declared.GetNestedScopes());
-	CopyUsedIds(/* out */ KeptParametrics, Used.ParametricTypes, Declared.GetParametricTypes().GetAllTypes());
+		// Copy nested scopes and parametric types
+		CopyUsedIds(/* out */ KeptScopes, Used.NestedScopes, Declared.GetNestedScopes());
+		CopyUsedIds(/* out */ KeptParametrics, Used.ParametricTypes, Declared.GetParametricTypes().GetAllTypes());
 
-	// Remap copied nested scopes
-	for (FNestedScope& KeptScope : KeptScopes)
-	{
-		KeptScope.Inner = Remap(KeptScope.Inner);
-		KeptScope.Outer = Remap(KeptScope.Outer);
-	}
+		// Remap copied nested scopes
+		for (FNestedScope& KeptScope : KeptScopes)
+		{
+			KeptScope.Inner = Remap(KeptScope.Inner);
+			KeptScope.Outer = Remap(KeptScope.Outer);
+		}
 
-	// Remap copied parametric types and copy parameters
-	TConstArrayView<FTypeId> DeclaredParameters = Declared.GetParametricTypes().GetAllParameters();
-	for (FParametricType& KeptType : KeptParametrics)
-	{
-		FParameterIndexRange OldParameters = KeptType.Parameters;
-		KeptType.Name = Remap(KeptType.Name);
-		KeptType.Parameters.Idx = KeptParameters.Num();
-		KeptParameters.Append(DeclaredParameters.Slice(OldParameters.Idx, OldParameters.NumParameters));
-	}
+		// Remap copied parametric types and copy parameters
+		TConstArrayView<FTypeId> DeclaredParameters = Declared.GetParametricTypes().GetAllParameters();
+		for (FParametricType& KeptType : KeptParametrics)
+		{
+			FParameterIndexRange OldParameters = KeptType.Parameters;
+			KeptType.Name = Remap(KeptType.Name);
+			KeptType.Parameters.Idx = KeptParameters.Num();
+			KeptParameters.Append(DeclaredParameters.Slice(OldParameters.Idx, OldParameters.NumParameters));
+		}
 
-	// Remap copied parameters
-	for (FTypeId& KeptParameter : KeptParameters)
-	{
-		KeptParameter = Remap(KeptParameter);
+		// Remap copied parameters
+		for (FTypeId& KeptParameter : KeptParameters)
+		{
+			KeptParameter = Remap(KeptParameter);
+		}
 	}
 }
 
@@ -263,11 +268,11 @@ static FOptionalSchemaId GetStaticInnerSchema(const FMemberSchema& Schema, const
 		check(IsStructOrEnum(InnermostType));
 		if (InnermostType.IsLeaf())
 		{
-			return ToOptional(NewIds.RemapEnum(InnerSchema.Get()));
+			return ToOptional(static_cast<FSchemaId>(NewIds.RemapEnum(InnerSchema.Get())));
 		}
 		else if (!InnermostType.AsStruct().IsDynamic)
 		{
-			return ToOptional(NewIds.RemapStruct(InnerSchema.Get()));
+			return ToOptional(static_cast<FSchemaId>(NewIds.RemapStruct(InnerSchema.Get())));
 		}
 	}
 
@@ -320,12 +325,12 @@ static ESuper GetInheritance(FOptionalSchemaId Super, TConstArrayView<const FMem
 
 static void WriteSchema(TArray64<uint8>& Out, const FBuiltStructSchema& Struct, const FWriteIds& NewIds)
 {
-	FOptionalSchemaId NewSuperId = Struct.Super ? ToOptional(NewIds.RemapStruct(Struct.Super.Get())) : NoId;
+	FOptionalSchemaId NewSuperId = Struct.Super ? ToOptionalSchema(NewIds.RemapStruct(Struct.Super.Get())) : NoId;
 	ESuper Inheritance = GetInheritance(NewSuperId, Struct.MemberSchemas);
 
 	TArray<FMemberType> MemberTypes = GetMemberTypes(Struct);
 	TArray<FMemberType> InnerRangeTypes = GetInnerRangeTypes(Struct);
-	TArray<FMemberId> MemberNames = RemapIds(NewIds, Struct.MemberNames);
+	TArray<FMemberId> MemberNames = NewIds.HasStableNames() ? RemapIds(NewIds, Struct.MemberNames) : Struct.MemberNames;
 	TArray<FSchemaId> InnerSchemas = GetInnerSchemas(Struct, NewIds, Inheritance);
 
 	check(MemberNames.Num() + UsesSuper(Inheritance) == MemberTypes.Num());
@@ -333,7 +338,7 @@ static void WriteSchema(TArray64<uint8>& Out, const FBuiltStructSchema& Struct,
 	// Zero init for determinism
 	alignas(FStructSchema) uint8 HeaderData[offsetof(FStructSchema, Footer)] = {};
 	FStructSchema& BinaryHeader = *reinterpret_cast<FStructSchema*>(HeaderData);
-	BinaryHeader.Type = NewIds.Remap(Struct.Type);
+	BinaryHeader.Type = NewIds.HasStableNames() ? NewIds.Remap(Struct.Type) : Struct.Type;
 	BinaryHeader.Inheritance = Inheritance;
 	BinaryHeader.IsDense = Struct.bDense;
 	BinaryHeader.NumMembers = IntCastChecked<uint16>(MemberTypes.Num());
@@ -557,9 +562,9 @@ public:
 	, Debug(InDebug)
 	{}
 
-	void WriteMembers(FStructSchemaId BuiltId, const FBuiltStruct& Struct)
+	FStructSchemaId WriteMembers(FStructSchemaId BuiltId, const FBuiltStruct& Struct)
 	{
-		FSchemaId WriteId = NewIds.RemapStruct(BuiltId);
+		FStructSchemaId WriteId = NewIds.RemapStruct(BuiltId);
 		const FBuiltStructSchema& Schema = Schemas.Structs[WriteId.Idx];
 		const TArray<FMemberId>& Order = Schema.MemberNames;
 		const int32 NumSuper = Schema.MemberSchemas.Num() - Schema.MemberNames.Num();
@@ -604,6 +609,8 @@ public:
 		}
 
 		Bits.Flush();
+
+		return WriteId;
 	}
 
 private:
@@ -735,10 +742,10 @@ private:
 
 ////////////////////////////////////////////////////////////////////////////////////////////////////////////////////////////////////////////////////
 
-FWriter::FWriter(const FIdIndexerBase& DeclaredIds, const FBuiltSchemas& InSchemas)
+FWriter::FWriter(const FIdIndexerBase& DeclaredIds, const FBuiltSchemas& InSchemas, ESchemaFormat Format)
 : Schemas(InSchemas)
 , Debug(DeclaredIds)
-, NewIds(new FWriteIds(DeclaredIds, InSchemas))
+, NewIds(new FWriteIds(DeclaredIds, InSchemas, Format))
 {}
 
 FWriter::~FWriter()
@@ -759,9 +766,9 @@ void FWriter::WriteSchemas(TArray64<uint8>& Out) const
 	WriteSchemasImpl(Out, Schemas, *NewIds);
 }
 
-void FWriter::WriteMembers(TArray64<uint8>& Out, FStructSchemaId BuiltId, const FBuiltStruct& Struct) const
+FStructSchemaId FWriter::WriteMembers(TArray64<uint8>& Out, FStructSchemaId BuiltId, const FBuiltStruct& Struct) const
 {
-	FMemberWriter(Out, Schemas, *NewIds, Debug).WriteMembers(BuiltId, Struct);
+	return FMemberWriter(Out, Schemas, *NewIds, Debug).WriteMembers(BuiltId, Struct);
 }
 
 ////////////////////////////////////////////////////////////////////////////////////////////////////////////////////////////////////////////////////
diff --git a/Engine/Plugins/Experimental/PlainProps/Source/Public/PlainPropsBind.h b/Engine/Plugins/Experimental/PlainProps/Source/Public/PlainPropsBind.h
--- a/Engine/Plugins/Experimental/PlainProps/Source/Public/PlainPropsBind.h
+++ b/Engine/Plugins/Experimental/PlainProps/Source/Public/PlainPropsBind.h
@@ -246,16 +246,16 @@ class FCustomBindings
 {
 public:
 	UE_NONCOPYABLE(FCustomBindings);
-	explicit FCustomBindings(const FDebugIds& In) : Debug(In) {}
+	FCustomBindings(const FDebugIds& Dbg, const FCustomBindings* InBase = nullptr) : Base(InBase), Debug(Dbg) {}
 
 	// @param Binding must outlive this or call DropStruct()
-	void						BindStruct(FStructSchemaId Id, ICustomBinding& Binding);
-	ICustomBinding*				FindStruct(FStructSchemaId Id)								{ return Find(Id); }
-	const ICustomBinding*		FindStruct(FStructSchemaId Id) const						{ return Find(Id); }
-	void						DropStruct(FStructSchemaId Id);
+	PLAINPROPS_API void						BindStruct(FStructSchemaId Id, ICustomBinding& Binding);
+	ICustomBinding*							FindStruct(FStructSchemaId Id)								{ return Find(Id); }
+	const ICustomBinding*					FindStruct(FStructSchemaId Id) const						{ return Find(Id); }
+	PLAINPROPS_API void						DropStruct(FStructSchemaId Id);
 
 private:
-	ICustomBinding*				Find(FStructSchemaId Id) const;
+	PLAINPROPS_API ICustomBinding* Find(FStructSchemaId Id) const;
 
 	struct FEntry
 	{
@@ -263,10 +263,17 @@ private:
 		ICustomBinding* Binding;
 	};
 
+	const FCustomBindings*					Base = nullptr;
 	TArray<FEntry, TInlineAllocator<8>>		Entries;
 	const FDebugIds&						Debug;
 };
 
+template<typename T, typename Ids>
+struct TCustomBind{ using Type = void; };
+
+template<typename T, typename Ids>
+using CustomBind = typename TCustomBind<T, Ids>::Type;
+
 ////////////////////////////////////////////////////////////////////////////////////////////////
 
 class FConstructionRequest
@@ -470,28 +477,16 @@ public:
 
 ////////////////////////////////////////////////////////////////////////////////////////////////
 
-// Todo: Switch to FPlatformMemory::KernelAddressBit once that is submitted
-//
-// This bit is always zero in user mode addresses and most likely won't be used by current or future
-// CPU features like ARM's PAC / Top-Byte Ignore or Intel's Linear Address Masking / 5-Level Paging
-#if defined(__x86_64__) || defined(_M_X64)
-	static constexpr uint32 KernelAddressBit = 63;
-#elif defined(__aarch64__) || defined(_M_ARM64)
-	static constexpr uint32 KernelAddressBit = 55;
-#else
-	#error Unsupported architecture, please declare which address bit distinguish user space from kernel space
-#endif
-
 class FRangeBinding
 {
 	static constexpr uint64		SizeMask = 0b1111;
-	static constexpr uint64		LeafMask = uint64(1) << KernelAddressBit;
+	static constexpr uint64		LeafMask = uint64(1) << FPlatformMemory::KernelAddressBit;
 	static constexpr uint64		BindMask = ~(SizeMask | LeafMask);
 	uint64						Handle;
 
 public:
-	FRangeBinding(const IItemRangeBinding& Binding, ERangeSizeType SizeType);
-	FRangeBinding(const ILeafRangeBinding& Binding, ERangeSizeType SizeType);
+	PLAINPROPS_API FRangeBinding(const IItemRangeBinding& Binding, ERangeSizeType SizeType);
+	PLAINPROPS_API FRangeBinding(const ILeafRangeBinding& Binding, ERangeSizeType SizeType);
 	
 	bool						IsLeafBinding() const		{ return !!(LeafMask & Handle); }
 	const IItemRangeBinding&	AsItemBinding() const		{ check(!IsLeafBinding()); return *reinterpret_cast<IItemRangeBinding*>(Handle & BindMask); }
@@ -522,15 +517,15 @@ class FSchemaBindings
 public:
 	UE_NONCOPYABLE(FSchemaBindings);
 	explicit FSchemaBindings(const FDebugIds& In) : Debug(In) {}
-	~FSchemaBindings();
+	PLAINPROPS_API ~FSchemaBindings();
 
-	void								BindStruct(FStructSchemaId Id, TConstArrayView<FMemberBinding> Schema);
-	const FSchemaBinding&				GetStruct(FStructSchemaId Id) const;
-	void								DropStruct(FStructSchemaId Id);
+	PLAINPROPS_API void						BindStruct(FStructSchemaId Id, TConstArrayView<FMemberBinding> Schema);
+	PLAINPROPS_API const FSchemaBinding&	GetStruct(FStructSchemaId Id) const;
+	PLAINPROPS_API void						DropStruct(FStructSchemaId Id);
 
 private:
-	TArray<TUniquePtr<FSchemaBinding>>	Bindings;
-	const FDebugIds&					Debug;
+	TArray<TUniquePtr<FSchemaBinding>>		Bindings;
+	const FDebugIds&						Debug;
 };
 
 ////////////////////////////////////////////////////////////////////////////////////////////////
@@ -578,15 +573,17 @@ FStructSchemaId BindCustomStructOnce()
 		using Ids = typename Runtime::Ids;
 
 		FBinding()
-		: Id(Runtime::GetTypes().DeclareStruct(Ids::IndexNativeType(Type::Name), NoId, CustomBinding::GetMemberIds(), CustomBinding::Occupancy))
 		{
-			Runtime::GetBindings().BindStruct(Id, Instance);
+			FTypeId Name = Ids::IndexNativeType(CttiOf<Type>::Name);
+			Id = Ids::IndexStruct(Name);
+			Runtime::GetTypes().DeclareStruct(Id, Name, CustomBinding::GetMemberIds(), CustomBinding::Occupancy);
+			Runtime::GetCustoms().BindStruct(Id, Instance);
 		}
 
 		~FBinding()
 		{
-			Runtime::GetBindings().DropStruct(Id);
-			Runtime::GetDeclarations().DropStruct(Id);
+			Runtime::GetCustoms().DropStruct(Id);
+			Runtime::GetTypes().DropStruct(Id);
 		}
 
 		CustomBinding Instance;
@@ -708,18 +705,25 @@ FMemberBinding BindMember()
 	else
 	{
 		using CustomBinding = typename Runtime::template CustomBindings<Type>::Type;
-		using RangeBinding = RangeBind<Type>;
-		constexpr uint32 NumRangeBindings = CountRangeBindings<RangeBinding>();
-
-		if constexpr (std::is_void_v<CustomBinding> && NumRangeBindings)
+		if constexpr (!std::is_void_v<CustomBinding>)
 		{
-			using InnermostType = typename TInnerType<RangeBinding, NumRangeBindings>::Type;
-			Out.RangeBindings = GetRangeBindings<RangeBinding, NumRangeBindings>();
-			Out.InnermostType = BindType<InnermostType, Ids>(Out.InnermostSchema);
+			Out.InnermostType = BindMemberStruct<Type, CustomBinding, Runtime>(Out.InnermostSchema);
 		}
 		else
 		{
-			Out.InnermostType = BindMemberStruct<Type, CustomBinding, Runtime>(Out.InnermostSchema);
+			using RangeBinding = RangeBind<Type>;
+			if constexpr (!std::is_void_v<RangeBinding>)
+			{
+				constexpr uint32 NumRangeBindings = CountRangeBindings<RangeBinding>();
+				using InnermostType = typename TInnerType<RangeBinding, NumRangeBindings>::Type;
+
+				Out.RangeBindings = GetRangeBindings<RangeBinding, NumRangeBindings>();
+				Out.InnermostType = BindType<InnermostType, Ids>(Out.InnermostSchema);
+			}
+			else
+			{
+				Out.InnermostType = BindMemberStruct<Type, CustomBinding, Runtime>(Out.InnermostSchema);
+			}
 		}
 	}
 	
@@ -830,7 +834,10 @@ struct FMemberBinder
 
 //////////////////////////////////////////////////////////////////////////
 
-// Save -> load ids
+// Save -> load struct ids for ESchemaFormat::InMemoryNames
+[[nodiscard]] PLAINPROPS_API TArray<FStructSchemaId> IndexInMemoryNames(const FSchemaBatch& Schemas,  FIdIndexerBase& Indexer);
+
+// Save -> load ids for ESchemaFormat::StableNames
 struct FIdBinding
 {
 	TConstArrayView<FNameId>			Names;
@@ -862,7 +869,7 @@ struct FIdTranslatorBase
 	static FIdBinding TranslateIds(FMutableMemoryView To, FIdIndexerBase& Indexer, TConstArrayView<FNameId> TranslatedNames, const FSchemaBatch& From);
 };
 
-// Maps saved ids -> runtime load ids
+// Maps saved ids -> runtime load ids for ESchemaFormat::StableNames
 struct FIdTranslator : FIdTranslatorBase
 {
 	template<class NameType>
diff --git a/Engine/Plugins/Experimental/PlainProps/Source/Public/PlainPropsBuild.h b/Engine/Plugins/Experimental/PlainProps/Source/Public/PlainPropsBuild.h
--- a/Engine/Plugins/Experimental/PlainProps/Source/Public/PlainPropsBuild.h
+++ b/Engine/Plugins/Experimental/PlainProps/Source/Public/PlainPropsBuild.h
@@ -18,6 +18,12 @@ struct FBuiltRange;
 class FDebugIds;
 struct FUnpackedLeafType;
 
+struct FBuiltStructDeleter 
+{
+	PLAINPROPS_API void operator()(FBuiltStruct* Ptr) const;
+};
+using FBuiltStructPtr = TUniquePtr<FBuiltStruct, FBuiltStructDeleter>;
+
 //////////////////////////////////////////////////////////////////////////
 
 struct FMemberSchema
@@ -88,7 +94,7 @@ PLAINPROPS_API FMemberSchema MakeNestedRangeSchema(ERangeSizeType SizeType, cons
 namespace Private
 {
 	[[nodiscard]] PLAINPROPS_API FBuiltRange*	BuildStructuralRange(/* in-out */ TArrayView64<FBuiltRange*> Structs);
-	[[nodiscard]] PLAINPROPS_API FBuiltRange*	BuildStructuralRange(/* in-out */ TArrayView64<TUniquePtr<FBuiltStruct>> Structs);
+	[[nodiscard]] PLAINPROPS_API FBuiltRange*	BuildStructuralRange(/* in-out */ TArrayView64<FBuiltStructPtr> Structs);
 	[[nodiscard]] PLAINPROPS_API FBuiltRange*	BuildLeafRange(FUnpackedLeafType Leaf, uint64 Num, FMemoryView Values);
 	PLAINPROPS_API void							NormalizeLeafRange(FUnpackedLeafType Leaf, FBuiltRange& Out);
 
@@ -127,7 +133,7 @@ template<typename EnumType, typename SizeType>
 	return { MakeEnumRangeSchema<EnumType, SizeType>(Enum), Private::BuildLeafRange(Values.GetData(), Values.Num()) };
 }
 
-[[nodiscard]] inline FTypedRange BuildStructRange(FStructSchemaId Schema, ERangeSizeType SizeType, /* in-out */ TArrayView64<TUniquePtr<FBuiltStruct>> Values )
+[[nodiscard]] inline FTypedRange BuildStructRange(FStructSchemaId Schema, ERangeSizeType SizeType, /* in-out */ TArrayView64<FBuiltStructPtr> Values )
 {
 	return { MakeStructRangeSchema(SizeType, Schema), Values.Num() ? Private::BuildStructuralRange(/* ownership xfer */ Values) : nullptr };
 }
@@ -158,13 +164,13 @@ public:
 	void AddEnum64(FMemberId Name, FEnumSchemaId Schema, uint64 Value)	{ AddLeaf(Name, {ELeafType::Enum, ELeafWidth::B64}, ToOptional(Schema), Value); }
 
 	PLAINPROPS_API void AddLeaf(FMemberId Name, FUnpackedLeafType Leaf, FOptionalEnumSchemaId Enum, uint64 Value);
-	PLAINPROPS_API void AddStruct(FMemberId Name, FStructSchemaId Schema, TUniquePtr<FBuiltStruct>&& Struct);	
+	PLAINPROPS_API void AddStruct(FMemberId Name, FStructSchemaId Schema, FBuiltStructPtr&& Struct);	
 	PLAINPROPS_API void AddRange(FMemberId Name, FTypedRange&& Range);
 	
 	// Build members into a single nested super struct member, no-op if no non-super members has been added
 	PLAINPROPS_API void BuildSuperStruct(const FStructDeclaration& Super, const FDebugIds& Debug);
 
-	[[nodiscard]] PLAINPROPS_API TUniquePtr<FBuiltStruct> BuildAndReset(const FStructDeclaration& Declared, const FDebugIds& Debug);
+	[[nodiscard]] PLAINPROPS_API FBuiltStructPtr BuildAndReset(const FStructDeclaration& Declared, const FDebugIds& Debug);
 
 	bool IsEmpty() const { return Members.IsEmpty(); }
 
diff --git a/Engine/Plugins/Experimental/PlainProps/Source/Public/PlainPropsBuildSchema.h b/Engine/Plugins/Experimental/PlainProps/Source/Public/PlainPropsBuildSchema.h
--- a/Engine/Plugins/Experimental/PlainProps/Source/Public/PlainPropsBuildSchema.h
+++ b/Engine/Plugins/Experimental/PlainProps/Source/Public/PlainPropsBuildSchema.h
@@ -46,13 +46,14 @@ public:
 	using FStructDeclarations = TConstArrayView<TUniquePtr<FStructDeclaration>>;
 	using FEnumDeclarations = TConstArrayView<TUniquePtr<FEnumDeclaration>>;
 
-	FSchemasBuilder(FStructDeclarations InStructs, FEnumDeclarations InEnums, const FDebugIds& InDebug);
-	~FSchemasBuilder();
+	PLAINPROPS_API explicit FSchemasBuilder(const FDeclarations& Declarations);
+	PLAINPROPS_API FSchemasBuilder(FStructDeclarations InStructs, FEnumDeclarations InEnums, const FDebugIds& InDebug);
+	PLAINPROPS_API ~FSchemasBuilder();
 
-	FEnumSchemaBuilder&							NoteEnum(FEnumSchemaId Id);
-	FStructSchemaBuilder&						NoteStruct(FStructSchemaId Id);
-	void										NoteStructAndMembers(FStructSchemaId Id, const FBuiltStruct& Struct);
-	FBuiltSchemas								Build();
+	PLAINPROPS_API FEnumSchemaBuilder&			NoteEnum(FEnumSchemaId Id);
+	PLAINPROPS_API FStructSchemaBuilder&		NoteStruct(FStructSchemaId Id);
+	PLAINPROPS_API void							NoteStructAndMembers(FStructSchemaId Id, const FBuiltStruct& Struct);
+	PLAINPROPS_API FBuiltSchemas				Build();
 
 private:
 	FStructDeclarations							DeclaredStructs;
diff --git a/Engine/Plugins/Experimental/PlainProps/Source/Public/PlainPropsCtti.h b/Engine/Plugins/Experimental/PlainProps/Source/Public/PlainPropsCtti.h
--- a/Engine/Plugins/Experimental/PlainProps/Source/Public/PlainPropsCtti.h
+++ b/Engine/Plugins/Experimental/PlainProps/Source/Public/PlainPropsCtti.h
@@ -2,27 +2,27 @@
 
 #pragma once
 
-#include <cstdint>
+#include "HAL/Platform.h"
 
 namespace PlainProps 
 {
 
-#define _PP_ARITHMETIC_CTTI(T) struct T##_Ctti  { inline static constexpr char Name[] = #T; using Type = T; }; T##_Ctti CttiOfPtr(T*);
-_PP_ARITHMETIC_CTTI(float)
-_PP_ARITHMETIC_CTTI(double)
-_PP_ARITHMETIC_CTTI(bool)
-_PP_ARITHMETIC_CTTI(int8_t)
-_PP_ARITHMETIC_CTTI(uint8_t)
-_PP_ARITHMETIC_CTTI(int16_t)
-_PP_ARITHMETIC_CTTI(int32_t)
-_PP_ARITHMETIC_CTTI(uint32_t)
-_PP_ARITHMETIC_CTTI(int64_t)
-_PP_ARITHMETIC_CTTI(uint64_t)
-_PP_ARITHMETIC_CTTI(char)
-_PP_ARITHMETIC_CTTI(char8_t)
-_PP_ARITHMETIC_CTTI(char16_t)
-_PP_ARITHMETIC_CTTI(char32_t)
-_PP_ARITHMETIC_CTTI(wchar_t)
+#define _PP_ARITHMETIC_CTTI(T, N) struct T##_Ctti  { inline static constexpr char Name[] = #N; using Type = T; }; T##_Ctti CttiOfPtr(T*);
+_PP_ARITHMETIC_CTTI(float, f32)
+_PP_ARITHMETIC_CTTI(double, f64)
+_PP_ARITHMETIC_CTTI(bool, bool)
+_PP_ARITHMETIC_CTTI(int8, s8)
+_PP_ARITHMETIC_CTTI(uint8, u8)
+_PP_ARITHMETIC_CTTI(int16, s16)
+_PP_ARITHMETIC_CTTI(uint16, u16)
+_PP_ARITHMETIC_CTTI(int32, s32)
+_PP_ARITHMETIC_CTTI(uint32, u32)
+_PP_ARITHMETIC_CTTI(int64, s64)
+_PP_ARITHMETIC_CTTI(uint64, u64)
+_PP_ARITHMETIC_CTTI(char, utf8)
+_PP_ARITHMETIC_CTTI(char8_t, utf8)
+_PP_ARITHMETIC_CTTI(char16_t, utf16)
+_PP_ARITHMETIC_CTTI(char32_t, utf32)
 #undef _PP_ARITHMETIC_CTTI
 
 //////////////////////////////////////////////////////////////////////////
@@ -59,6 +59,7 @@ template<typename T> using CttiOf = typename TCttiOf<T>::Type;
 #define PP_REFLECT_ENUM(NS, T, ...)						_PP_REFLECT_ENUM(_PP_NUM_ARGS(__VA_ARGS__), NS, T, __VA_ARGS__)
 #define PP_REFLECT_STRUCT(NS, T, S, ...)				_PP_REFLECT_STRUCT(_PP_NUM_ARGS(__VA_ARGS__), NS, T, S, __VA_ARGS__)
 #define PP_REFLECT_STRUCT_TEMPLATE(NS, T, S, ...)		_PP_REFLECT_STRUCT_TEMPLATE(_PP_NUM_ARGS(__VA_ARGS__), NS, T, S, __VA_ARGS__)
+#define PP_NAME_STRUCT(NS, T)							struct T##_Ctti { inline static constexpr char Name[] = #T; using Type = NS :: T; }; T##_Ctti CttiOfPtr(T*);
 
 // Alternate set of macros to reflect classes/structs with bitfield bool members, e.g. uint8 bOol : 1;
 
diff --git a/Engine/Plugins/Experimental/PlainProps/Source/Public/PlainPropsDeclare.h b/Engine/Plugins/Experimental/PlainProps/Source/Public/PlainPropsDeclare.h
--- a/Engine/Plugins/Experimental/PlainProps/Source/Public/PlainPropsDeclare.h
+++ b/Engine/Plugins/Experimental/PlainProps/Source/Public/PlainPropsDeclare.h
@@ -54,8 +54,8 @@ public:
 	UE_NONCOPYABLE(FDeclarations);
 	explicit FDeclarations(const FDebugIds& In) : Debug(In) {}
 
-	void											DeclareEnum(FEnumSchemaId Id, FTypeId Type, EEnumMode Mode, ELeafWidth Width, TConstArrayView<FEnumerator> Enumerators);
-	void											DeclareStruct(FStructSchemaId Id, FTypeId Type, TConstArrayView<FMemberId> MemberOrder, EMemberPresence Occupancy, FOptionalStructSchemaId Super = {});
+	PLAINPROPS_API void								DeclareEnum(FEnumSchemaId Id, FTypeId Type, EEnumMode Mode, ELeafWidth Width, TConstArrayView<FEnumerator> Enumerators);
+	PLAINPROPS_API void								DeclareStruct(FStructSchemaId Id, FTypeId Type, TConstArrayView<FMemberId> MemberOrder, EMemberPresence Occupancy, FOptionalStructSchemaId Super = {});
 	
 	void											DropEnum(FEnumSchemaId Id)		{ Check(Id); DeclaredEnums[Id.Idx].Reset(); }
 	void											DropStruct(FStructSchemaId Id)	{ Check(Id); DeclaredStructs[Id.Idx].Reset(); }
@@ -73,8 +73,8 @@ protected:
 	const FDebugIds&								Debug;
 
 #if DO_CHECK
-	void											Check(FEnumSchemaId Id) const;
-	void											Check(FStructSchemaId Id) const;
+	PLAINPROPS_API void								Check(FEnumSchemaId Id) const;
+	PLAINPROPS_API void								Check(FStructSchemaId Id) const;
 #else
 	void											Check(...) const {}
 #endif
diff --git a/Engine/Plugins/Experimental/PlainProps/Source/Public/PlainPropsIndex.h b/Engine/Plugins/Experimental/PlainProps/Source/Public/PlainPropsIndex.h
--- a/Engine/Plugins/Experimental/PlainProps/Source/Public/PlainPropsIndex.h
+++ b/Engine/Plugins/Experimental/PlainProps/Source/Public/PlainPropsIndex.h
@@ -11,9 +11,9 @@ namespace PlainProps
 class FNestedScopeIndexer
 {
 public:
-	FNestedScopeId						Index(FNestedScope Scope);
+	PLAINPROPS_API FNestedScopeId		Index(FNestedScope Scope);
 	FNestedScopeId						Index(FScopeId Outer, FFlatScopeId Inner)			{ return Index({Outer, Inner}); }
-	FNestedScope						Resolve(FNestedScopeId Id) const;
+	PLAINPROPS_API FNestedScope			Resolve(FNestedScopeId Id) const;
 	int32								Num() const											{ return Scopes.Num(); }
 
 	auto								begin() const										{ return Scopes.begin(); }
@@ -27,11 +27,11 @@ private:
 class FParametricTypeIndexer
 {
 public:
-	~FParametricTypeIndexer();
+	PLAINPROPS_API ~FParametricTypeIndexer();
 
-	FParametricTypeId					Index(FParametricTypeView View);
+	PLAINPROPS_API FParametricTypeId	Index(FParametricTypeView View);
 	/// @return View invalidated by calling Index() (switch Parameters to TPagedArray to avoid)
-	FParametricTypeView					Resolve(FParametricTypeId Id) const;
+	PLAINPROPS_API FParametricTypeView	Resolve(FParametricTypeId Id) const;
 	FParametricType						At(int32 Idx) const									{ return Types[Idx]; }
 	int32								Num() const											{ return Types.Num(); }
 	TConstArrayView<FParametricType>	GetAllTypes() const									{ return Types; }
@@ -49,12 +49,12 @@ private:
 class FIdIndexerBase : public FDebugIds
 {
 public:
-	FScopeId							NestScope(FScopeId Outer, FFlatScopeId Inner);
-	FParametricTypeId					MakeParametricTypeId(FConcreteTypenameId Name, TConstArrayView<FTypeId> Params);
-	FTypeId								MakeParametricType(FTypeId Type, TConstArrayView<FTypeId> Params);
+	PLAINPROPS_API FScopeId				NestScope(FScopeId Outer, FFlatScopeId Inner);
+	PLAINPROPS_API FParametricTypeId	MakeParametricTypeId(FConcreteTypenameId Name, TConstArrayView<FTypeId> Params);
+	PLAINPROPS_API FTypeId				MakeParametricType(FTypeId Type, TConstArrayView<FTypeId> Params);
 
-	FEnumSchemaId						IndexEnum(FTypeId Type);
-	FStructSchemaId						IndexStruct(FTypeId Type);
+	PLAINPROPS_API FEnumSchemaId		IndexEnum(FTypeId Type);
+	PLAINPROPS_API FStructSchemaId		IndexStruct(FTypeId Type);
 	
 	virtual FNestedScope				Resolve(FNestedScopeId Id) const override final		{ return NestedScopes.Resolve(Id); }
 	virtual FParametricTypeView			Resolve(FParametricTypeId Id) const override final	{ return ParametricTypes.Resolve(Id); }
diff --git a/Engine/Plugins/Experimental/PlainProps/Source/Public/PlainPropsLoad.h b/Engine/Plugins/Experimental/PlainProps/Source/Public/PlainPropsLoad.h
--- a/Engine/Plugins/Experimental/PlainProps/Source/Public/PlainPropsLoad.h
+++ b/Engine/Plugins/Experimental/PlainProps/Source/Public/PlainPropsLoad.h
@@ -5,6 +5,7 @@
 #include "Containers/ArrayView.h"
 #include "PlainPropsRead.h"
 #include "Memory/MemoryFwd.h"
+#include "Templates/UniquePtr.h"
 
 namespace PlainProps 
 {
@@ -16,26 +17,24 @@ class FRangeBinding;
 class FSchemaBindings;
 struct FTypedRange;
 
+struct FLoadBatchDeleter 
+{
+	PLAINPROPS_API void operator()(FLoadBatch* Ptr) const;
+};
+using FLoadBatchPtr = TUniquePtr<FLoadBatch, FLoadBatchDeleter>;
 
-PLAINPROPS_API FLoadBatch* CreateLoadPlans(FReadBatchId ReadId, const FDeclarations& Declarations, const FCustomBindings& Customs, const FSchemaBindings& Schemas, TConstArrayView<FStructSchemaId> RuntimeIds);
-PLAINPROPS_API void DestroyLoadPlans(FLoadBatch* Batch);
-PLAINPROPS_API void LoadStruct(uint8* Dst, FByteReader Src, FStructSchemaId Id, const FLoadBatch& Batch);
-PLAINPROPS_API void LoadStruct(uint8* Dst, FStructView Src, const FLoadBatch& Batch);
-PLAINPROPS_API void ConstructAndLoadStruct(uint8* Dst, FByteReader Src, FStructSchemaId Id, const FLoadBatch& Batch);
-PLAINPROPS_API void ConstructAndLoadStruct(uint8* Dst, FStructView Src, const FLoadBatch& Batch);
-PLAINPROPS_API void LoadRange(uint8* Dst, FRangeView Src, ERangeSizeType MaxType, TConstArrayView<FRangeBinding> Bindings, const FLoadBatch& Batch);
-
-//using FItemConstructor = void (*)(FLoadRangeContext& Ctx);
-
-//struct FLoadRangeBinding
-//{
-//	FItemConstructor	MakeItems;
-//	const void*			UserContext;
-//};
-
-// Possible optimization - opt-in leaf range load method to avoid copying and repeated calls when target representation is not a range
-//template<typename T> using TDirectLeafLoader = void (*)(TRangeView<T> Items);
-//using FDirectBoolLoader = void (*)(FBoolRangeView Bits);
+template<class Runtime>
+[[nodiscard]] FLoadBatchPtr CreateLoadPlans(FReadBatchId ReadId, TConstArrayView<FStructSchemaId> RuntimeIds)
+{
+	return CreateLoadPlans(ReadId, Runtime::GetTypes(), Runtime::GetCustoms(), Runtime::GetSchemas(), RuntimeIds);
+}
+
+PLAINPROPS_API FLoadBatchPtr				CreateLoadPlans(FReadBatchId ReadId, const FDeclarations& Declarations, const FCustomBindings& Customs, const FSchemaBindings& Schemas, TConstArrayView<FStructSchemaId> RuntimeIds);
+PLAINPROPS_API void							LoadStruct(void* Dst, FByteReader Src, FStructSchemaId Id, const FLoadBatch& Batch);
+PLAINPROPS_API void							LoadStruct(void* Dst, FStructView Src, const FLoadBatch& Batch);
+PLAINPROPS_API void							ConstructAndLoadStruct(void* Dst, FByteReader Src, FStructSchemaId Id, const FLoadBatch& Batch);
+PLAINPROPS_API void							ConstructAndLoadStruct(void* Dst, FStructView Src, const FLoadBatch& Batch);
+PLAINPROPS_API void							LoadRange(void* Dst, FRangeView Src, ERangeSizeType MaxType, TConstArrayView<FRangeBinding> Bindings, const FLoadBatch& Batch);
 
 } // namespace PlainProps
 
diff --git a/Engine/Plugins/Experimental/PlainProps/Source/Public/PlainPropsRead.h b/Engine/Plugins/Experimental/PlainProps/Source/Public/PlainPropsRead.h
--- a/Engine/Plugins/Experimental/PlainProps/Source/Public/PlainPropsRead.h
+++ b/Engine/Plugins/Experimental/PlainProps/Source/Public/PlainPropsRead.h
@@ -26,15 +26,15 @@ class FReadBatchId
 };
 
 // @param Schemas must outlive read batch
-const FSchemaBatch*		ValidateSchemas(FMemoryView Schemas);
-FReadBatchId			MountReadSchemas(const FSchemaBatch* Schemas);
-const FSchemaBatch*		UnmountReadSchemas(FReadBatchId Batch);
-
-uint32					NumStructSchemas(FReadBatchId Batch);
-const FStructSchema&	ResolveStructSchema(FReadBatchId Batch, FStructSchemaId Id);
-const FEnumSchema&		ResolveEnumSchema(FReadBatchId Batch, FEnumSchemaId Id);
-FNestedScope			ResolveUntranslatedNestedScope(FReadBatchId Batch, FNestedScopeId Id);
-FParametricTypeView		ResolveUntranslatedParametricType(FReadBatchId Batch, FParametricTypeId Id);
+PLAINPROPS_API const FSchemaBatch*		ValidateSchemas(FMemoryView Schemas);
+PLAINPROPS_API FReadBatchId				MountReadSchemas(const FSchemaBatch* Schemas);
+PLAINPROPS_API const FSchemaBatch*		UnmountReadSchemas(FReadBatchId Batch);
+
+PLAINPROPS_API uint32					NumStructSchemas(FReadBatchId Batch);
+PLAINPROPS_API const FStructSchema&		ResolveStructSchema(FReadBatchId Batch, FStructSchemaId Id);
+PLAINPROPS_API const FEnumSchema&		ResolveEnumSchema(FReadBatchId Batch, FEnumSchemaId Id);
+PLAINPROPS_API FNestedScope				ResolveUntranslatedNestedScope(FReadBatchId Batch, FNestedScopeId Id);
+PLAINPROPS_API FParametricTypeView		ResolveUntranslatedParametricType(FReadBatchId Batch, FParametricTypeId Id);
 
 //////////////////////////////////////////////////////////////////////////
 
diff --git a/Engine/Plugins/Experimental/PlainProps/Source/Public/PlainPropsSave.h b/Engine/Plugins/Experimental/PlainProps/Source/Public/PlainPropsSave.h
--- a/Engine/Plugins/Experimental/PlainProps/Source/Public/PlainPropsSave.h
+++ b/Engine/Plugins/Experimental/PlainProps/Source/Public/PlainPropsSave.h
@@ -3,25 +3,36 @@
 #pragma once
 
 #include "PlainPropsTypes.h"
+#include "PlainPropsBuild.h"
 #include "Templates/UniquePtr.h"
 
 namespace PlainProps 
 {
 
-struct FBuiltStruct;
 class FCustomBindings;
 class FDeclarations;
 class FSchemaBindings;
 
+// Temporary data structure, will be replaced by something more sophisticated
+// perhaps deduplicating all zero-memory defaults
+struct FDefaultStruct { FStructSchemaId Id; const void* Struct; };
+using FDefaultStructs = TConstArrayView<FDefaultStruct>;
+
 struct FSaveContext
 {
 	const FDeclarations&		Declarations;
 	const FSchemaBindings&		Schemas;
 	FCustomBindings&			Customs;
+	FDefaultStructs				Defaults;
 };
 
-[[nodiscard]] TUniquePtr<FBuiltStruct> SaveStruct(const void* Struct, FStructSchemaId Id, const FSaveContext& Context);
-[[nodiscard]] TUniquePtr<FBuiltStruct> SaveStructDelta(const void* Struct, const void* Default, FStructSchemaId Id, const FSaveContext& Context);
+template<typename Runtime>
+FSaveContext MakeSaveContext(FDefaultStructs Defaults)
+{
+	return { Runtime::GetTypes(), Runtime::GetSchemas(), Runtime::GetCustoms(), Defaults };
+}
 
-} // namespace PlainProps
+[[nodiscard]] PLAINPROPS_API FBuiltStructPtr SaveStruct(const void* Struct, FStructSchemaId Id, const FSaveContext& Context);
+[[nodiscard]] PLAINPROPS_API FBuiltStructPtr SaveStructDelta(const void* Struct, const void* Default, FStructSchemaId Id, const FSaveContext& Context);
 
+} // namespace PlainProps
\ No newline at end of file
diff --git a/Engine/Plugins/Experimental/PlainProps/Source/Public/PlainPropsTypes.h b/Engine/Plugins/Experimental/PlainProps/Source/Public/PlainPropsTypes.h
--- a/Engine/Plugins/Experimental/PlainProps/Source/Public/PlainPropsTypes.h
+++ b/Engine/Plugins/Experimental/PlainProps/Source/Public/PlainPropsTypes.h
@@ -268,6 +268,8 @@ using FOptionalParametricTypeId = TOptionalId<FParametricTypeId>;
 
 template<class IdType>
 inline constexpr TOptionalId<IdType> ToOptional(IdType Id) { return Id; }
+inline constexpr FOptionalSchemaId ToOptionalSchema(FEnumSchemaId Id) { return static_cast<FSchemaId>(Id); }
+inline constexpr FOptionalSchemaId ToOptionalSchema(FStructSchemaId Id) { return static_cast<FSchemaId>(Id); }
 
 //////////////////////////////////////////////////////////////////////////
 
@@ -379,11 +381,12 @@ public:
 	virtual FTypeId						Resolve(FStructSchemaId Id) const = 0;
 
 	virtual void						AppendDebugString(FString& Out, FNameId Name) const = 0;
-	virtual void						AppendDebugString(FString& Out, FScopeId Scope) const;
-	virtual void						AppendDebugString(FString& Out, FTypenameId Typename) const;
-	virtual void						AppendDebugString(FString& Out, FTypeId Type) const;
-	virtual void						AppendDebugString(FString& Out, FEnumSchemaId Name) const;
-	virtual void						AppendDebugString(FString& Out, FStructSchemaId Name) const;
+
+	PLAINPROPS_API virtual void			AppendDebugString(FString& Out, FScopeId Scope) const;
+	PLAINPROPS_API virtual void			AppendDebugString(FString& Out, FTypenameId Typename) const;
+	PLAINPROPS_API virtual void			AppendDebugString(FString& Out, FTypeId Type) const;
+	PLAINPROPS_API virtual void			AppendDebugString(FString& Out, FEnumSchemaId Name) const;
+	PLAINPROPS_API virtual void			AppendDebugString(FString& Out, FStructSchemaId Name) const;
 
 	PLAINPROPS_API FString				Print(FNameId Name) const;
 	PLAINPROPS_API FString				Print(FMemberId Name) const;
diff --git a/Engine/Plugins/Experimental/PlainProps/Source/Public/PlainPropsUeCoreBindings.h b/Engine/Plugins/Experimental/PlainProps/Source/Public/PlainPropsUeCoreBindings.h
--- a/Engine/Plugins/Experimental/PlainProps/Source/Public/PlainPropsUeCoreBindings.h
+++ b/Engine/Plugins/Experimental/PlainProps/Source/Public/PlainPropsUeCoreBindings.h
@@ -3,21 +3,13 @@
 #pragma once
 
 #include "Containers/Set.h"
+#include "Math/MathFwd.h"
 #include "PlainPropsBind.h"
 #include "PlainPropsLoad.h"
 #include "PlainPropsRead.h"
 #include "PlainPropsIndex.h"
 #include "UObject/NameTypes.h"
 
-namespace PlainProps
-{
-	template<>
-	void AppendString(FString& Out, const FName& Name) { Name.AppendString(Out); }
-}
-
-// Todo: PP_NAME_STRUCT(, FName);
-struct FName_Ctti { static constexpr char Name[] = "FName"; };
-FName_Ctti CttiOfPtr(FName*);
 
 namespace PlainProps::UE
 {
@@ -243,7 +235,6 @@ struct TUniquePtrBinding : public IItemRangeBinding
 		const TUniquePtr<T>& Ptr = Ctx.Request.GetRange<TUniquePtr<T>>();
 		Ctx.Items.SetAll(Ptr.Get(), Ptr ? 1 : 0);
 	}
-
 };
 
 //////////////////////////////////////////////////////////////////////////
@@ -352,6 +343,63 @@ struct TSetBinding : public IItemRangeBinding
 };
 
 //////////////////////////////////////////////////////////////////////////
+
+struct FTransformIds
+{
+	FMemberId Translate[3];
+	FMemberId Rotate[4];
+	FMemberId Scale[3];
+
+	TConstArrayView<FMemberId> GetAll() const
+	{
+		static_assert(offsetof(FTransformIds, Rotate)	== offsetof(FTransformIds, Translate) + sizeof(Translate));
+		static_assert(offsetof(FTransformIds, Scale)	== offsetof(FTransformIds, Rotate) + sizeof(Rotate));
+		return MakeArrayView(&Translate[0], 3 + 4 + 3);
+	}
+	
+	template<class Ids>	
+	static const FTransformIds& Get()
+	{
+		static FTransformIds Out = {	//Ids::IndexMember("Mask"),
+			{ Ids::IndexMember("TranslateX"), Ids::IndexMember("TranslateY"), Ids::IndexMember("TranslateZ") },
+			{ Ids::IndexMember("RotateX"), Ids::IndexMember("RotateY"),	Ids::IndexMember("RotateZ"), Ids::IndexMember("RotateW") },
+			{ Ids::IndexMember("ScaleX"), Ids::IndexMember("ScaleY"), Ids::IndexMember("ScaleZ") }};
+
+		return Out;
+	}
+};
+
+struct FTransformBinding : public ICustomBinding
+{
+	PLAINPROPS_API void			Save(FMemberBuilder& Dst, const FTransform& Src, const FTransform* Default, const FTransformIds& Ids) const;
+	PLAINPROPS_API void			Load(FTransform& Dst, FStructView Src, ECustomLoadMethod Method, const FLoadBatch& Batch, const FTransformIds& Ids) const;
+	PLAINPROPS_API virtual bool	DiffStruct(const void* StructA, const void* StructB) const override;
+};
+
+template<class Ids>
+struct TTransformBinding final : public FTransformBinding
+{
+	using Type = FTransform;
+	static constexpr EMemberPresence Occupancy = EMemberPresence::AllowSparse;
+
+	static TConstArrayView<FMemberId> GetMemberIds()
+	{
+		return FTransformIds::Get<Ids>().GetAll();
+	}
+
+	virtual void SaveStruct(FMemberBuilder& Dst, const void* Src, const void* Default, const FDebugIds& Debug) override
+	{
+		Save(Dst, *static_cast<const FTransform*>(Src), static_cast<const FTransform*>(Default), FTransformIds::Get<Ids>());
+	}
+
+	virtual void LoadStruct(void* Dst, FStructView Src, ECustomLoadMethod Method, const FLoadBatch& Batch) const override
+	{
+		Load(*static_cast<FTransform*>(Dst), Src, Method, Batch, FTransformIds::Get<Ids>());
+	}
+};
+
+//////////////////////////////////////////////////////////////////////////
+
 struct FSetOps
 {
 	union
@@ -557,7 +605,7 @@ struct TSetDeltaBinding : public ICustomBinding
 //		}
 //	}
 //
-//	virtual TUniquePtr<FBuiltStruct>	SaveStruct(const void* Src, const FDebugIds& Debug) const override
+//	virtual FBuiltStructPtr	SaveStruct(const void* Src, const FDebugIds& Debug) const override
 //	{
 //		...
 //	}
@@ -565,29 +613,45 @@ struct TSetDeltaBinding : public ICustomBinding
 
 }
 
+
+PP_NAME_STRUCT(, FName);
+
+namespace UE::Math { PP_NAME_STRUCT(, FTransform); }
+
+
 namespace PlainProps
 {
-	template<typename T>
-	struct TRangeBind<TArray<T>>
-	{
-		using Type = UE::TArrayBinding<T>;
-	};
 
-	template<>
-	struct TRangeBind<FString>
-	{
-		using Type = UE::FStringBinding;
-	};
+template<>
+PLAINPROPS_API void AppendString(FString& Out, const FName& Name);
 
-	template<typename T>
-	struct TRangeBind<TUniquePtr<T>>
-	{
-		using Type = UE::TUniquePtrBinding<T>;
-	};
+template<typename T>
+struct TRangeBind<TArray<T>>
+{
+	using Type = UE::TArrayBinding<T>;
+};
 
-	template<typename T>
-	struct TRangeBind<TSet<T>>
-	{
-		using Type = UE::TSetBinding<T>;
-	};
+template<>
+struct TRangeBind<FString>
+{
+	using Type = UE::FStringBinding;
+};
+
+template<typename T>
+struct TRangeBind<TUniquePtr<T>>
+{
+	using Type = UE::TUniquePtrBinding<T>;
+};
+
+template<typename T>
+struct TRangeBind<TSet<T>>
+{
+	using Type = UE::TSetBinding<T>;
+};
+
+template<typename Ids>
+struct TCustomBind<FTransform, Ids>
+{
+	using Type = UE::TTransformBinding<Ids>;
+};
 }
\ No newline at end of file
diff --git a/Engine/Plugins/Experimental/PlainProps/Source/Public/PlainPropsWrite.h b/Engine/Plugins/Experimental/PlainProps/Source/Public/PlainPropsWrite.h
--- a/Engine/Plugins/Experimental/PlainProps/Source/Public/PlainPropsWrite.h
+++ b/Engine/Plugins/Experimental/PlainProps/Source/Public/PlainPropsWrite.h
@@ -16,19 +16,19 @@ class FNestedScopeIndexer;
 class FParametricTypeIndexer;
 struct FWriteIds;
 
+enum class ESchemaFormat { StableNames, InMemoryNames };
+
 class FWriter
 {
 public:
-	PLAINPROPS_API FWriter(const FIdIndexerBase& DeclaredIds, const FBuiltSchemas& InSchemas);
+	PLAINPROPS_API FWriter(const FIdIndexerBase& DeclaredIds, const FBuiltSchemas& InSchemas, ESchemaFormat Format);
 	PLAINPROPS_API ~FWriter();
 	
 	PLAINPROPS_API bool								Uses(FNameId BuiltId) const;
-	PLAINPROPS_API TConstArrayView<FNestedScope>	GetUsedScopes() const;
-	PLAINPROPS_API TConstArrayView<FParametricType>	GetUsedParametrics() const;
 	PLAINPROPS_API FOptionalStructSchemaId			GetWriteId(FStructSchemaId BuiltId) const;
 
 	PLAINPROPS_API void								WriteSchemas(TArray64<uint8>& Out) const;
-	PLAINPROPS_API void								WriteMembers(TArray64<uint8>& Out, FStructSchemaId BuiltId, const FBuiltStruct& Struct) const;
+	PLAINPROPS_API FStructSchemaId					WriteMembers(TArray64<uint8>& Out, FStructSchemaId BuiltId, const FBuiltStruct& Struct) const;
 
 private:
 	const FBuiltSchemas&							Schemas;
diff --git a/Engine/Plugins/Experimental/PlainProps/Source/Tests/PlainPropsCttiTests.cpp b/Engine/Plugins/Experimental/PlainProps/Source/Tests/PlainPropsCttiTests.cpp
--- a/Engine/Plugins/Experimental/PlainProps/Source/Tests/PlainPropsCttiTests.cpp
+++ b/Engine/Plugins/Experimental/PlainProps/Source/Tests/PlainPropsCttiTests.cpp
@@ -8,7 +8,7 @@ namespace PlainProps::Test
 {
 
 static_assert(std::is_same_v<CttiOf<float>::Type, float>);
-static_assert(std::string_view(CttiOf<int>::Name) == std::string_view("int32_t"));
+static_assert(std::string_view(CttiOf<int>::Name) == std::string_view("s32"));
 
 //////////////////////////////////////////////////////////////////////////
 
diff --git a/Engine/Plugins/Experimental/PlainProps/Source/Tests/PlainPropsTests.cpp b/Engine/Plugins/Experimental/PlainProps/Source/Tests/PlainPropsTests.cpp
--- a/Engine/Plugins/Experimental/PlainProps/Source/Tests/PlainPropsTests.cpp
+++ b/Engine/Plugins/Experimental/PlainProps/Source/Tests/PlainPropsTests.cpp
@@ -123,7 +123,7 @@ public:
 	TArray64<uint8>				Write();
 
 private:
-	TArray<TPair<FStructSchemaId, TUniquePtr<FBuiltStruct>>>	Objects;
+	TArray<TPair<FStructSchemaId, FBuiltStructPtr>>	Objects;
 	FDeclarations												Declarations;
 
 	TArray<FMemberId> NameMembers(std::initializer_list<const char*> Members)
@@ -185,14 +185,14 @@ TArray64<uint8> FTestBatchBuilder::Write()
 {
 	// Build partial schemas
 	FSchemasBuilder SchemaBuilders(Declarations.GetStructs(), Declarations.GetEnums(), *this);
-	for (const TPair<FStructSchemaId, TUniquePtr<FBuiltStruct>>& Object : Objects)
+	for (const TPair<FStructSchemaId, FBuiltStructPtr>& Object : Objects)
 	{
 		SchemaBuilders.NoteStructAndMembers(Object.Key, *Object.Value);
 	}
 	FBuiltSchemas Schemas = SchemaBuilders.Build(); 
 
 	// Filter out declared but unused names and ids
-	FWriter Writer(*this, Schemas);
+	FWriter Writer(*this, Schemas, ESchemaFormat::StableNames);
 
 	// Write names
 	TArray64<uint8> Out;
@@ -219,7 +219,7 @@ TArray64<uint8> FTestBatchBuilder::Write()
 
 	// Write objects
 	WriteU32(Out, TestMagics[2]);
-	for (const TPair<FStructSchemaId, TUniquePtr<FBuiltStruct>>& Object : Objects)
+	for (const TPair<FStructSchemaId, FBuiltStructPtr>& Object : Objects)
 	{
 		WriteU32(/* out */ Tmp, TestMagics[3]);
 		WriteU32(/* out */ Tmp, Writer.GetWriteId(Object.Key).Get().Idx);
@@ -628,14 +628,14 @@ TEST_CASE_NAMED(FPlainPropsReadWriteTest, "System::Core::Serialization::PlainPro
 
 			FMemberBuilder Members;
 			Members.Add(	Batch.NameMember("I1"), 100);
-			TUniquePtr<FBuiltStruct> NestedInStruct = Members.BuildAndReset(Batch.Get(NestedId), Batch);
+			FBuiltStructPtr NestedInStruct = Members.BuildAndReset(Batch.Get(NestedId), Batch);
 
 			Members.AddStruct(Batch.NameMember("Nested"), NestedId, MoveTemp(NestedInStruct));
 			Members.Add(Batch.NameMember("Leaf"), true);
-			TUniquePtr<FBuiltStruct> Struct = Members.BuildAndReset(Batch.Get(StructId), Batch);
+			FBuiltStructPtr Struct = Members.BuildAndReset(Batch.Get(StructId), Batch);
 		
 			Members.Add(Batch.NameMember("I2"), 200);
-			TUniquePtr<FBuiltStruct> NestedInObject = Members.BuildAndReset(Batch.Get(NestedId), Batch);
+			FBuiltStructPtr NestedInObject = Members.BuildAndReset(Batch.Get(NestedId), Batch);
 
 			Members.Add(Batch.NameMember("L1"), 123.f);
 			Members.AddStruct(Batch.NameMember("S"), StructId, MoveTemp(Struct));
diff --git a/Engine/Plugins/Experimental/PlainProps/Source/Tests/PlainPropsUeCoreTests.cpp b/Engine/Plugins/Experimental/PlainProps/Source/Tests/PlainPropsUeCoreTests.cpp
--- a/Engine/Plugins/Experimental/PlainProps/Source/Tests/PlainPropsUeCoreTests.cpp
+++ b/Engine/Plugins/Experimental/PlainProps/Source/Tests/PlainPropsUeCoreTests.cpp
@@ -15,6 +15,7 @@
 #include "Algo/Compare.h"
 #include "Containers/StringView.h"
 #include "Containers/Map.h"
+#include "Math/Transform.h"
 #include "Templates/UnrealTemplate.h"
 #include "Tests/TestHarnessAdapter.h"
 
@@ -24,6 +25,7 @@ namespace PlainProps::UE::Test
 static TIdIndexer<FName>	GNames;
 static FDeclarations		GTypes(/* debug */ GNames);
 static FSchemaBindings		GSchemas(/* debug */ GNames);
+static FCustomBindings		GCustoms(/* debug */ GNames);
 
 struct FIds
 {
@@ -44,7 +46,7 @@ struct FIds
 template<class T>
 struct TCustomBindings
 {
-	using Type = void;
+	using Type = CustomBind<T, FIds>;
 };
 
 ////////////////////////////////////////////////////////////////////////////////////////////////
@@ -54,8 +56,9 @@ struct FDefaultRuntime
 	using Ids = FIds;
 	template<class T> using CustomBindings = TCustomBindings<T>;
 
-	static FDeclarations&			GetDeclarations()				{ return GTypes; }
-	static FSchemaBindings&			GetSchemas()					{ return GSchemas; }
+	static FDeclarations&			GetTypes()			{ return GTypes; }
+	static FSchemaBindings&			GetSchemas()		{ return GSchemas; }
+	static FCustomBindings&			GetCustoms()		{ return GCustoms; }
 };
 
 ////////////////////////////////////////////////////////////////////////////////////////////////
@@ -67,8 +70,8 @@ struct TScopedEnumDeclaration
 	using Ctti = CttiOf<Enum>;
 
 	FEnumSchemaId Id;
-	TScopedEnumDeclaration() : Id(DeclareNativeEnum<Ctti, Ids>(Runtime::GetDeclarations(), Mode)) {}
-	~TScopedEnumDeclaration() { Runtime::GetDeclarations().DropEnum(Id); }
+	TScopedEnumDeclaration() : Id(DeclareNativeEnum<Ctti, Ids>(Runtime::GetTypes(), Mode)) {}
+	~TScopedEnumDeclaration() { Runtime::GetTypes().DropEnum(Id); }
 };
 
 template<class T, EMemberPresence Occupancy = EMemberPresence::AllowSparse, class Runtime = FDefaultRuntime>
@@ -79,17 +82,17 @@ struct TScopedStructDeclaration
 	FStructSchemaId Id;
 
 	TScopedStructDeclaration()
-	: Id(DeclareNativeStruct<CttiOf<T>, Ids>(Runtime::GetDeclarations(), Occupancy))
+	: Id(DeclareNativeStruct<CttiOf<T>, Ids>(Runtime::GetTypes(), Occupancy))
 	{}
 
 	~TScopedStructDeclaration()
 	{
-		Runtime::GetDeclarations().DropStruct(Id);
+		Runtime::GetTypes().DropStruct(Id);
 	}
 
 	const FStructDeclaration& Get() const
 	{
-		return Runtime::GetDeclarations().Get(Id);
+		return Runtime::GetTypes().Get(Id);
 	}
 };
 
@@ -179,14 +182,14 @@ public:
 	TArray64<uint8>				Write() const;
 
 private:
-	using IdBuiltStructPair = TPair<FStructSchemaId, TUniquePtr<FBuiltStruct>>;
+	using IdBuiltStructPair = TPair<FStructSchemaId, FBuiltStructPtr>;
 	TArray<IdBuiltStructPair>	SavedObjects;
 	FNameBinding				SavedNames;
 	FCustomBindings				Customs;
 };
 
 FBatchSaver::FBatchSaver()
-: Customs(/* debug */ GNames)
+: Customs(/* debug */ GNames, &GCustoms)
 {
 	Customs.BindStruct(SavedNames.Declaration.Id, SavedNames);
 }
@@ -202,7 +205,7 @@ template<class T>
 bool FBatchSaver::SaveDelta(const T& Object, const T& Default) 
 {
 	FStructSchemaId Id = IndexNativeStruct<T, FIds>();
-	if (TUniquePtr<FBuiltStruct> Delta = SaveStructDelta(&Object, &Default, Id, {GTypes, GSchemas, Customs}))
+	if (FBuiltStructPtr Delta = SaveStructDelta(&Object, &Default, Id, {GTypes, GSchemas, Customs}))
 	{
 		SavedObjects.Emplace(Id, MoveTemp(Delta));
 		return true;
@@ -235,7 +238,7 @@ TArray64<uint8> FBatchSaver::Write() const
 	FBuiltSchemas Schemas = SchemaBuilders.Build(); 
 
 	// Filter out declared but unused names and ids
-	FWriter Writer(GNames, Schemas);
+	FWriter Writer(GNames, Schemas, ESchemaFormat::StableNames);
 	TArray<FName> UsedNames;
 	for (uint32 Idx = 0, Num = GNames.NumNames(); Idx < Num; ++Idx)
 	{
@@ -260,7 +263,7 @@ TArray64<uint8> FBatchSaver::Write() const
 
 	// Write objects
 	WriteU32(Out, Magics[2]);
-	for (const TPair<FStructSchemaId, TUniquePtr<FBuiltStruct>>& Object : SavedObjects)
+	for (const TPair<FStructSchemaId, FBuiltStructPtr>& Object : SavedObjects)
 	{
 		WriteU32(/* out */ Tmp, Magics[3]);
 		WriteU32(/* out */ Tmp, Writer.GetWriteId(Object.Key).Get().Idx);
@@ -286,7 +289,7 @@ class FBatchLoader
 {
 public:
 	FBatchLoader(FMemoryView Data)
-	: Customs(/* debug */ GNames)
+	: Customs(/* debug */ GNames, &GCustoms)
 	{
 		// Read ids
 		FByteReader It(Data);
@@ -331,7 +334,7 @@ public:
 	~FBatchLoader()
 	{
 		CHECK(LoadIdx == Objects.Num()); // Test should load all saved objects
-		DestroyLoadPlans(Plans);
+		Plans.Reset();
 		const FSchemaBatch* LoadSchemas = UnmountReadSchemas(Objects[0].Schema.Batch);
 		DestroyTranslatedSchemas(LoadSchemas);
 	}
@@ -355,7 +358,7 @@ private:
 	TConstArrayView<FName>		Ids;
 	FNameBinding				Names;
 	FCustomBindings				Customs;
-	FLoadBatch*					Plans;
+	FLoadBatchPtr				Plans;
 	TArray<FStructView>			Objects;
 	int32						LoadIdx = 0;
 };
@@ -756,9 +759,39 @@ TEST_CASE_NAMED(FPlainPropsUeCoreTest, "System::Core::Serialization::PlainProps:
 			});
 	}
 
+	SECTION("Transform")
+	{
+		BindCustomStructOnce<TTransformBinding<FIds>, FDefaultRuntime>();
+
+		Run([](FBatchSaver& Batch)
+			{
+				CHECK(!Batch.SaveDelta(FTransform(),FTransform()));
+				CHECK(!Batch.SaveDelta(FTransform(FVector::UnitY()), FTransform(FVector::UnitY())));
+
+				Batch.Save(FTransform());
+
+				// This should only save translation
+				Batch.SaveDelta(FTransform(FVector::UnitY()), FTransform());
+			}, 
+			[](FBatchLoader& Batch)
+			{
+				CHECK(Batch.Load<FTransform>().Equals(FTransform(), 0.0));
+
+				FTransform TranslateOnly;
+				TranslateOnly.SetTranslation(FVector(10, 20, 30));
+				TranslateOnly.SetRotation(FQuat(1, 2, 3, 4));
+				TranslateOnly.SetScale3D(FVector(5, 6, 7));
+				Batch.LoadInto(TranslateOnly);
+				CHECK(TranslateOnly.GetTranslation() == FVector::UnitY());
+				CHECK(TranslateOnly.GetRotation() == FQuat(1, 2, 3, 4));
+				CHECK(TranslateOnly.GetScale3D() == FVector(5, 6, 7));
+			});
+	}
+
 	SECTION("Reference")
 	{}
 }
 
 } // namespace PlainProps::UE::Test
+
 #endif // WITH_TESTS
\ No newline at end of file
