PlainProps optimize build phase by moving intermediate representation over to a scratch allocator

#rnx

diff --git a/Engine/Plugins/Experimental/PlainProps/Source/Private/PlainPropsBind.cpp b/Engine/Plugins/Experimental/PlainProps/Source/Private/PlainPropsBind.cpp
--- a/Engine/Plugins/Experimental/PlainProps/Source/Private/PlainPropsBind.cpp
+++ b/Engine/Plugins/Experimental/PlainProps/Source/Private/PlainPropsBind.cpp
@@ -95,7 +95,7 @@ FRangeMemberBinding FMemberVisitor::GrabRange()
 	FOptionalSchemaId InnermostSchema = HasSchema(InnerTypes.Last()) ? ToOptional(GrabInnerSchema()) : NoId; 
 	uint64 Offset = GrabMemberOffset();
 		
-	return { &InnerTypes[0], RangeBindings, static_cast<uint32>(InnerTypes.Num()), InnermostSchema, Offset};
+	return { &InnerTypes[0], RangeBindings, static_cast<uint16>(InnerTypes.Num()), InnermostSchema, Offset};
 }
 
 void FMemberVisitor::SkipMember()
@@ -141,7 +141,7 @@ FRangeBinding::FRangeBinding(const ILeafRangeBinding& Binding, ERangeSizeType Si
 void* FLeafRangeAllocator::Allocate(FUnpackedLeafType Leaf, uint64 Num)
 {
 	check(!Range);
-	Range = FBuiltRange::Create(Num, SizeOf(Leaf.Width));
+	Range = FBuiltRange::Create(Scratch, Num, SizeOf(Leaf.Width));
 	return Range->Data;
 }
 
diff --git a/Engine/Plugins/Experimental/PlainProps/Source/Private/PlainPropsBuild.cpp b/Engine/Plugins/Experimental/PlainProps/Source/Private/PlainPropsBuild.cpp
--- a/Engine/Plugins/Experimental/PlainProps/Source/Private/PlainPropsBuild.cpp
+++ b/Engine/Plugins/Experimental/PlainProps/Source/Private/PlainPropsBuild.cpp
@@ -11,114 +11,114 @@
 namespace PlainProps
 {
 
-void FBuiltStructDeleter::operator()(FBuiltStruct* Ptr) const
+uint8* FScratchAllocator::AllocateInNewPage(SIZE_T Size, uint32 Alignment)
 {
-	delete Ptr;
-}
-
-FBuiltRange* FBuiltRange::Create(uint64 NumItems, SIZE_T ItemSize)
-{
-	check(NumItems > 0);
-	FBuiltRange* Out = new (FMemory::Malloc(sizeof(FBuiltRange) + NumItems * ItemSize)) FBuiltRange;
-	Out->Num = NumItems;
-	return Out;
-}
-
-//////////////////////////////////////////////////////////////////////////
-
-FMemberSchema MakeNestedRangeSchema(ERangeSizeType SizeType, const FMemberSchema& InnerRangeSchema)
-{
-	FMemberSchema Out = { FMemberType(SizeType), InnerRangeSchema.InnerSchema };
-	Out.InnerRangeTypes.Reserve(1 + InnerRangeSchema.InnerRangeTypes.Num());
-	Out.InnerRangeTypes.Add(InnerRangeSchema.Type);
-	Out.InnerRangeTypes.Append(InnerRangeSchema.InnerRangeTypes);
-	return Out;
-}
-
-//////////////////////////////////////////////////////////////////////////
-
-template<typename BuiltType>
-FBuiltRange* BuildStructuralRangeImpl(/* in-out */ TArrayView64<BuiltType> Values)
-{
-	static_assert(alignof(FBuiltRange) >= alignof(BuiltType));
-	FBuiltRange* Out = FBuiltRange::Create(Values.Num(), sizeof(FBuiltRange));
-	BuiltType* DataIt = reinterpret_cast<BuiltType*>(Out->Data);
-	for (BuiltType& Value : Values)
+	if (Size >= DataSize || (DataSize - Size) < static_cast<uint32>(PageEnd - Cursor))
 	{
-		new (DataIt++) BuiltType(MoveTemp(Value));
+		SIZE_T LonePageSize = Align(offsetof(FPage, Data) + Size, Alignment);
+		FPage*& PrevPage = LastPage ? LastPage->PrevPage : LastPage;
+		FPage Header = { PrevPage };
+		PrevPage = new (FMemory::Malloc(LonePageSize)) FPage { Header };
+		return Align(PrevPage->Data, Alignment);
 	}
-
+	
+	FPage Header = { LastPage };
+	LastPage = new (FMemory::Malloc(PageSize)) FPage { Header };
+	uint8* Out = Align(LastPage->Data, Alignment);
+	Cursor = Out + Size;
+	PageEnd = LastPage->Data + DataSize;
+	check(Cursor <= PageEnd);
 	return Out;
 }
 
-namespace Private
+FScratchAllocator::~FScratchAllocator()
 {
-	FBuiltRange* BuildStructuralRange(TArrayView64<FBuiltRange*> Values)
+	for (FPage* It = LastPage; It; )
 	{
-		return PlainProps::BuildStructuralRangeImpl(Values);
-	}
-
-	FBuiltRange* BuildStructuralRange(TArrayView64<FBuiltStructPtr> Values)
-	{
-		return PlainProps::BuildStructuralRangeImpl(Values);
+		FPage* Prev = It->PrevPage;
+		FMemory::Free(It);
+		It = Prev;
 	}
+}
 
-	FBuiltRange* BuildLeafRange(FUnpackedLeafType Leaf, uint64 Num, FMemoryView Values)
-	{
-		check(Values.GetSize() == Num * SizeOf(Leaf.Width));
-		FBuiltRange* Out = FBuiltRange::Create(Num, SizeOf(Leaf.Width));
-		FMemory::Memcpy(Out->Data, Values.GetData(), Values.GetSize());
-		return Out;
-	}
-	
-	template<typename FloatType>
-	void NormalizeFloats(FloatType* Values, uint64 Num)
+FMemberType& FMemberSchema::EditInnermostType(FScratchAllocator& Scratch)
+{
+	if (NumInnerRanges > 1)
 	{
-		for (uint64 Idx = 0; Idx < Num; ++Idx)
-		{
-			// Reject NaN / INF and ignore negative zero for now
-			checkf(FMath::IsFinite(Values[Idx]), TEXT("Saving NaN or INF isn't supported"));
-		}
+		FMemberType* Clone = Scratch.AllocateArray<FMemberType>(NumInnerRanges);
+		FMemory::Memcpy(Clone, NestedRangeTypes, NumInnerRanges * sizeof(FMemberType));
+		NestedRangeTypes = Clone;
+		return Clone[NumInnerRanges - 1];
 	}
 
-	void NormalizeLeafRange(FUnpackedLeafType Leaf, FBuiltRange& Out)
-	{
-		check(Leaf.Type == ELeafType::Float);
-		if (Leaf.Width == ELeafWidth::B32)
-		{
-			NormalizeFloats(reinterpret_cast<float*>(Out.Data), Out.Num);
-		}
-		else
-		{
-			check(Leaf.Width == ELeafWidth::B64);
-			NormalizeFloats(reinterpret_cast<double*>(Out.Data), Out.Num);
-		}
-	}
-} // namespace Private
+	return NumInnerRanges == 0 ? Type : InnerRangeType;
+}
 
 //////////////////////////////////////////////////////////////////////////
 
-FMemberBuilder::FMemberBuilder() {}
-FMemberBuilder::~FMemberBuilder() {}
-
-void FMemberBuilder::AddStruct(FMemberId Name, FStructSchemaId Schema, FBuiltStructPtr&& Struct)
+FBuiltRange* FBuiltRange::Create(FScratchAllocator& Scratch, uint64 NumItems, SIZE_T ItemSize)
 {
-	Members.Emplace(Name, Schema, MoveTemp(Struct));
+	check(NumItems > 0);
+	FBuiltRange* Out = new (Scratch.Allocate(sizeof(FBuiltRange) + NumItems * ItemSize, alignof(FBuiltRange))) FBuiltRange;
+	Out->Num = NumItems;
+	return Out;
 }
 
-void FMemberBuilder::AddLeaf(FMemberId Name, FUnpackedLeafType Leaf, FOptionalEnumSchemaId Enum, uint64 Value)
+//////////////////////////////////////////////////////////////////////////
+
+FMemberSchema MakeNestedRangeSchema(FScratchAllocator& Scratch, ERangeSizeType SizeType, FMemberSchema InnerRangeSchema)
 {
-	Members.Emplace(Name, Leaf, Enum, Value);
+	check(InnerRangeSchema.NumInnerRanges > 0);
+	uint16 NumInnerRanges = IntCastChecked<uint16>(1 + InnerRangeSchema.NumInnerRanges);
+	FMemberType* InnerRangeTypes = Scratch.AllocateArray<FMemberType>(NumInnerRanges);
+	InnerRangeTypes[0] = InnerRangeSchema.Type;
+	FMemory::Memcpy(&InnerRangeTypes[1], InnerRangeSchema.GetInnerRangeTypes().GetData(), InnerRangeSchema.NumInnerRanges * sizeof(FMemberType));
+
+	return { FMemberType(SizeType), InnerRangeSchema.Type, NumInnerRanges, InnerRangeSchema.InnerSchema, InnerRangeTypes };
 }
 
-void FMemberBuilder::AddRange(FMemberId Name, FTypedRange&& Range)
+//////////////////////////////////////////////////////////////////////////
+
+FBuiltRange* CloneLeaves(FScratchAllocator& Scratch, uint64 Num, const void* InData, SIZE_T LeafSize)
 {
-	Members.Emplace(Name, MoveTemp(Range));
+	if (Num == 0)
+	{
+		return nullptr;
+	}
+
+	FBuiltRange* Out = FBuiltRange::Create(Scratch, Num, LeafSize);
+	FMemory::Memcpy(Out->Data, InData, Num * LeafSize);
+	return Out;
 }
 
-TSet<FBuiltStruct*> GLiveStructsFoo;
+//	template<typename FloatType>
+//	void NormalizeFloats(FloatType* Values, uint64 Num)
+//	{
+//		for (uint64 Idx = 0; Idx < Num; ++Idx)
+//		{
+//			// Reject NaN / INF and ignore negative zero for now
+//			checkf(FMath::IsFinite(Values[Idx]), TEXT("Saving NaN or INF isn't supported"));
+//		}
+//	}
+//
+//	void NormalizeLeafRange(FUnpackedLeafType Leaf, FBuiltRange& Out)
+//	{
+//		check(Leaf.Type == ELeafType::Float);
+//		if (Leaf.Width == ELeafWidth::B32)
+//		{
+//			NormalizeFloats(reinterpret_cast<float*>(Out.Data), Out.Num);
+//		}
+//		else
+//		{
+//			check(Leaf.Width == ELeafWidth::B64);
+//			NormalizeFloats(reinterpret_cast<double*>(Out.Data), Out.Num);
+//		}
+//	}
+//} // namespace Private
+
+//////////////////////////////////////////////////////////////////////////
 
-void FMemberBuilder::BuildSuperStruct(const FStructDeclaration& Super, const FDebugIds& Debug)
+void FMemberBuilder::BuildSuperStruct(FScratchAllocator& Scratch, const FStructDeclaration& Super, const FDebugIds& Debug)
 {
 	// If we need to support EMemberPresence::RequireAll for sub structs,
 	// we need access to struct declaration here or create an empty super
@@ -128,12 +128,12 @@ void FMemberBuilder::BuildSuperStruct(const FStructDeclaration& Super, const FDe
 		return;
 	}
 	
-	FBuiltStructPtr OnlyMember = BuildAndReset(Super, Debug);
+	FBuiltStructPtr OnlyMember = BuildAndReset(Scratch, Super, Debug);
 	Members.Emplace(FBuiltMember::MakeSuper(Super.Id, MoveTemp(OnlyMember)));
 	check(Members[0].Schema.Type.AsStruct().IsSuper);
 }
 
-FBuiltStructPtr FMemberBuilder::BuildAndReset(const FStructDeclaration& Declared, const FDebugIds& Debug)
+FBuiltStructPtr FMemberBuilder::BuildAndReset(FScratchAllocator& Scratch, const FStructDeclaration& Declared, const FDebugIds& Debug)
 {
 	checkf(!(Declared.Super && Declared.Occupancy == EMemberPresence::RequireAll),
 		TEXT("Requiring sub structs to be dense isn't implemented"));
@@ -157,7 +157,7 @@ FBuiltStructPtr FMemberBuilder::BuildAndReset(const FStructDeclaration& Declared
 
 	uint32 Num = static_cast<uint32>(Members.Num());
 	SIZE_T NumBytes = sizeof(FBuiltStruct) + Num * sizeof(FBuiltMember);
-	FBuiltStruct* Out = reinterpret_cast<FBuiltStruct*>(FMemory::MallocZeroed(NumBytes, alignof(FBuiltStruct)));
+	FBuiltStruct* Out = reinterpret_cast<FBuiltStruct*>(Scratch.AllocateZeroed(NumBytes, alignof(FBuiltStruct)));
 	Out->NumMembers = IntCastChecked<uint16>(Num);
 	for (FBuiltMember& Member : Members)
 	{
@@ -189,108 +189,48 @@ uint64 ValueCast(double Value)
 	return CheckFiniteBitCast<uint64>(Value);
 }
 
-
 //////////////////////////////////////////////////////////////////////////
 
-FBuiltMember::FBuiltMember(FBuiltMember&& O)
-: Name(O.Name)
-, Schema(MoveTemp(O.Schema))
-, Value(O.Value)
+template<typename SchemaIdType>
+static FMemberSchema MakeMemberSchema(FMemberType Type, SchemaIdType InnerSchema)
 {
-	O.Name = NoId;
-	O.Schema.Type = FMemberType(ELeafType::Bool, ELeafWidth::B8);
-	O.Value.Leaf = 0;
+	return {Type, Type, 0, FOptionalSchemaId(InnerSchema), nullptr};
 }
 
 FBuiltMember::FBuiltMember(FMemberId Name, FUnpackedLeafType Leaf, FOptionalEnumSchemaId Enum, uint64 Value)
-: FBuiltMember(Name, {Leaf.Pack(), FOptionalSchemaId(Enum)}, { .Leaf = Value})
+: FBuiltMember(Name, MakeMemberSchema(Leaf.Pack(), Enum), { .Leaf = Value})
 {}
 
-//FBuiltMember::FBuiltMember(FMemberId Name, FEnumSchemaId Schema, ELeafWidth Width, uint64 Value)
-//: FBuiltMember(Name, MakeMemberSchema(Width, Schema), { .Leaf = Value })
-//{}
-
-FBuiltMember::FBuiltMember(FMemberId Name, FTypedRange&& Range)
+FBuiltMember::FBuiltMember(FMemberId Name, FTypedRange Range)
 : FBuiltMember(Name, MoveTemp(Range.Schema), { .Range = Range.Values })
 {}
 
-FBuiltMember::FBuiltMember(FMemberId Name, FStructSchemaId Schema, FBuiltStructPtr&& Value)
-: FBuiltMember(Name, {DefaultStructType, FOptionalSchemaId(Schema)}, { .Struct = Value.Release() })
+FBuiltMember::FBuiltMember(FMemberId Name, FStructSchemaId Schema, FBuiltStructPtr Value)
+: FBuiltMember(Name, MakeMemberSchema(DefaultStructType, Schema), { .Struct = Value })
 {}
 
-FBuiltMember FBuiltMember::MakeSuper(FStructSchemaId Schema, FBuiltStructPtr&& Value)
+FBuiltMember FBuiltMember::MakeSuper(FStructSchemaId Schema, FBuiltStructPtr Value)
 {
-	return FBuiltMember(NoId, {SuperStructType, FOptionalSchemaId(Schema)}, { .Struct = Value.Release() });
-}
-
-FBuiltMember::~FBuiltMember()
-{
-	if (Schema.Type.GetKind() == EMemberKind::Struct)
-	{
-		delete Value.Struct;
-	}
-	else if (Schema.Type.GetKind() == EMemberKind::Range)
-	{
-		FBuiltRange::Delete(Value.Range, Schema.InnerSchema, Schema.InnerRangeTypes);
-	}
-}
-
-FBuiltMember& FBuiltMember::operator=(FBuiltMember&& O)
-{
-	this->~FBuiltMember();
-	new (this) FBuiltMember(MoveTemp(O));
-	return *this;
+	return FBuiltMember(NoId, MakeMemberSchema(SuperStructType, Schema), { .Struct = Value });
 }
 
 //////////////////////////////////////////////////////////////////////////
 
-uint64 FBuiltRange::Delete(FBuiltRange* Range, FOptionalSchemaId InnerSchema, TConstArrayView<FMemberType> InnerTypes)
-{ 
-	// TODO: Handle struct and nested ranges
-
-	switch (InnerTypes[0].GetKind()) 
-	{
-	case EMemberKind::Struct:
-		//DeleteStructs(InnerTypes[0].AsStruct(), static_cast<FStructSchemaId>(Schema.InnerSchema.Get()), Range);
-		break;
-	case EMemberKind::Range:
-		//for (FNestedRangeIterator It(Range); It;)
-		//{
-		//	FBuiltRange::Delete(&*It, InnerSchema, InnerTypes.RightChop(1));
-		//	It.Advance();
-		//}
-	case EMemberKind::Leaf:		break;
-	}
-
-	FMemory::Free(Range);
-
-	return 0; // todo: return number of bytes if needed?
-}
-
-//////////////////////////////////////////////////////////////////////////
-
-FBuiltStruct::~FBuiltStruct()
-{ 
-	//check(GLiveStructsFoo.Remove(this) == 1);
-
-	for (const FBuiltMember& Member : MakeArrayView(Members, NumMembers))
-	{
-		Member.~FBuiltMember();
-	}
-}
-
-//////////////////////////////////////////////////////////////////////////
-
-FTypedRange FStructRangeBuilder::BuildAndReset(const FStructDeclaration& Declared, const FDebugIds& Debug)
+FTypedRange FStructRangeBuilder::BuildAndReset(FScratchAllocator& Scratch, const FStructDeclaration& Declared, const FDebugIds& Debug)
 {
-	TArray64<FBuiltStructPtr> BuiltStructs;
-	BuiltStructs.Reserve(Structs.Num());
-	for (FMemberBuilder& Struct : Structs)
+	FTypedRange Out = { MakeStructRangeSchema(SizeType, Declared.Id) };
+
+	if (int64 Num = Structs.Num())
 	{
-		BuiltStructs.Emplace(Struct.BuildAndReset(Declared, Debug));
-	}
+		Out.Values = FBuiltRange::Create(Scratch, Structs.Num(), sizeof(FBuiltStruct*));
+		FBuiltStruct** OutIt = reinterpret_cast<FBuiltStruct**>(Out.Values->Data);
+		for (FMemberBuilder& Struct : Structs)
+		{
+			*OutIt++ = Struct.BuildAndReset(Scratch, Declared, Debug);
+		}
+	}		
 
-	return BuildStructRange(Declared.Id, SizeType, /* move */ BuiltStructs);
+	return Out;
 }
 
 //////////////////////////////////////////////////////////////////////////
@@ -300,12 +240,17 @@ FNestedRangeBuilder::~FNestedRangeBuilder()
 	checkf(Ranges.IsEmpty(), TEXT("Half-built range, forgot to call BuildAndReset() before destruction?"));
 }
 
-FTypedRange FNestedRangeBuilder::BuildAndReset(ERangeSizeType SizeType)
+FTypedRange FNestedRangeBuilder::BuildAndReset(FScratchAllocator& Scratch, ERangeSizeType SizeType)
 {
-	FTypedRange Out = { MakeNestedRangeSchema(SizeType, Schema), 
-						Ranges.IsEmpty() ? nullptr : Private::BuildStructuralRange(Ranges) };
-	Ranges.Reset();
-	return Out;
+	FBuiltRange* Out = nullptr;
+	if (int64 Num = Ranges.Num())
+	{
+		Out = FBuiltRange::Create(Scratch, Num, sizeof(FBuiltRange*));
+		FMemory::Memcpy(Out->Data, Ranges.GetData(), Ranges.NumBytes());
+		Ranges.Reset();
+	}
+
+	return { MakeNestedRangeSchema(Scratch, SizeType, Schema), Out };
 }
 
 } // namespace PlainProps
\ No newline at end of file
diff --git a/Engine/Plugins/Experimental/PlainProps/Source/Private/PlainPropsBuildSchema.cpp b/Engine/Plugins/Experimental/PlainProps/Source/Private/PlainPropsBuildSchema.cpp
--- a/Engine/Plugins/Experimental/PlainProps/Source/Private/PlainPropsBuildSchema.cpp
+++ b/Engine/Plugins/Experimental/PlainProps/Source/Private/PlainPropsBuildSchema.cpp
@@ -39,7 +39,7 @@ static FString PrintMemberSchema(FMemberType Type, FOptionalSchemaId InnerSchema
 
 static FString PrintMemberSchema(FMemberSchema Schema)
 {
-	return PrintMemberSchema(Schema.Type, Schema.InnerSchema, Schema.InnerRangeTypes);
+	return PrintMemberSchema(Schema.Type, Schema.InnerSchema, Schema.GetInnerRangeTypes());
 }
 
 //////////////////////////////////////////////////////////////////////////
@@ -48,6 +48,7 @@ struct FStructSchemaBuilder
 {
 	const FStructDeclaration&					Declaration;
 	FSchemasBuilder&							AllSchemas;
+	FScratchAllocator&							Scratch;
 	const FDebugIds&							Debug;
 	TMap<FOptionalMemberId, FMemberSchema>		NotedMembers;
 	bool										bMissingMemberNoted = false;
@@ -69,13 +70,14 @@ struct FEnumSchemaBuilder
 
 //////////////////////////////////////////////////////////////////////////
 
-FSchemasBuilder::FSchemasBuilder(const FDeclarations& Types)
-: FSchemasBuilder(Types.GetStructs(), Types.GetEnums(), Types.GetDebug())
+FSchemasBuilder::FSchemasBuilder(const FDeclarations& Types, FScratchAllocator& InScratch)
+: FSchemasBuilder(Types.GetStructs(), Types.GetEnums(), Types.GetDebug(), InScratch)
 {}
 
-FSchemasBuilder::FSchemasBuilder(FStructDeclarations InStructs, FEnumDeclarations InEnums, const FDebugIds& InDebug)
+FSchemasBuilder::FSchemasBuilder(FStructDeclarations InStructs, FEnumDeclarations InEnums, const FDebugIds& InDebug, FScratchAllocator& InScratch)
 : DeclaredStructs(InStructs)
 , DeclaredEnums(InEnums)
+, Scratch(InScratch)
 , Debug(InDebug)
 {
 	StructIndices.Init(INDEX_NONE, InStructs.Num());
@@ -107,7 +109,7 @@ FStructSchemaBuilder& FSchemasBuilder::NoteStruct(FStructSchemaId Id)
 {
 	checkf(!bBuilt, TEXT("Noted new members after building"));
 	checkf(DeclaredStructs[Id.Idx], TEXT("Undeclared struct '%s' noted"), *Debug.Print(Id));
-	return GetOrEmplace(StructIndices[Id.Idx], Structs, *DeclaredStructs[Id.Idx], *this, Debug);
+	return GetOrEmplace(StructIndices[Id.Idx], Structs, *DeclaredStructs[Id.Idx], *this, Scratch, Debug);
 }
 
 void FSchemasBuilder::NoteStructAndMembers(FStructSchemaId Id, const FBuiltStruct& Struct)
@@ -145,7 +147,7 @@ void FSchemasBuilder::NoteInheritanceChains()
 		for (FOptionalStructSchemaId Super = Structs[Idx].Declaration.Super; Super; Super = DeclaredStructs[Super.Get().Idx]->Super)
 		{
 			uint32 SuperIdx = Super.Get().Idx;
-			GetOrEmplace(StructIndices[SuperIdx], Structs, *DeclaredStructs[SuperIdx], *this, Debug);
+			GetOrEmplace(StructIndices[SuperIdx], Structs, *DeclaredStructs[SuperIdx], *this, Scratch, Debug);
 		}
 	}
 }
@@ -160,11 +162,11 @@ static bool RequiresDynamicStructSchema(const FMemberSchema& A, const FMemberSch
 		{
 			return true;
 		}
-		else if (A.Type.IsRange() && A.InnerRangeTypes.Last().IsStruct() && B.InnerRangeTypes.Last().IsStruct())
+		else if (A.Type.IsRange() && A.GetInnermostType().IsStruct() && B.GetInnermostType().IsStruct())
 		{
 			// Same range size and nested range sizes
-			return	A.Type == B.Type &&	Algo::Compare(	MakeArrayView(A.InnerRangeTypes).LeftChop(1),
-														MakeArrayView(B.InnerRangeTypes).LeftChop(1));
+			return	A.Type == B.Type &&	Algo::Compare(	A.GetInnerRangeTypes().LeftChop(1),
+														B.GetInnerRangeTypes().LeftChop(1));
 		}
 	}
 
@@ -194,8 +196,13 @@ void FStructSchemaBuilder::NoteMembersRecursively(const FBuiltStruct& Struct)
 		{
 			if (RequiresDynamicStructSchema(*Schema, Member.Schema))
 			{
-				SetIsDynamic(Schema->Type.IsStruct() ? Schema->Type : Schema->InnerRangeTypes.Last());
-				Schema->InnerSchema = {};
+				if (!Schema->GetInnermostType().AsStruct().IsDynamic)
+				{
+					SetIsDynamic(Schema->EditInnermostType(Scratch));
+					Schema->InnerSchema = NoId;
+				}
+				check(Schema->InnerSchema == NoId);
+				
 			}
 			else
 			{
@@ -222,14 +229,14 @@ void FStructSchemaBuilder::NoteMembersRecursively(const FBuiltStruct& Struct)
 		}
 		else 
 		{
-			check(IsStructOrEnum(Schema.InnerRangeTypes.Last()) == !!Schema.InnerSchema);
+			check(IsStructOrEnum(Schema.GetInnermostType()) == !!Schema.InnerSchema);
 
 			if (Schema.InnerSchema)
 			{
-				void* InnerSchemaBuilder = NoteStructOrEnum(AllSchemas, Schema.InnerRangeTypes.Last().IsStruct(), Schema.InnerSchema.Get());
+				void* InnerSchemaBuilder = NoteStructOrEnum(AllSchemas, Schema.GetInnermostType().IsStruct(), Schema.InnerSchema.Get());
 				if (Member.Value.Range)
 				{
-					NoteRangeRecursively(Schema.Type.AsRange().MaxSize, Schema.InnerRangeTypes, InnerSchemaBuilder, *Member.Value.Range);			
+					NoteRangeRecursively(Schema.Type.AsRange().MaxSize, Schema.GetInnerRangeTypes(), InnerSchemaBuilder, *Member.Value.Range);			
 				}
 			}
 		}
@@ -263,7 +270,7 @@ void FStructSchemaBuilder::NoteRangeRecursively(ERangeSizeType NumType, TConstAr
 	switch (Type.GetKind())
 	{
 	case EMemberKind::Struct:
-		for (const TUniquePtr<const FBuiltStruct>& Struct : Range.AsStructs())
+		for (const FBuiltStruct* Struct : Range.AsStructs())
 		{
 			static_cast<FStructSchemaBuilder*>(InnermostSchema)->NoteMembersRecursively(*Struct);
 		}
diff --git a/Engine/Plugins/Experimental/PlainProps/Source/Private/PlainPropsInternalBuild.h b/Engine/Plugins/Experimental/PlainProps/Source/Private/PlainPropsInternalBuild.h
--- a/Engine/Plugins/Experimental/PlainProps/Source/Private/PlainPropsInternalBuild.h
+++ b/Engine/Plugins/Experimental/PlainProps/Source/Private/PlainPropsInternalBuild.h
@@ -9,55 +9,25 @@
 namespace PlainProps
 {
 
-union FBuiltValue
-{
-	uint64			Leaf;
-	FBuiltStruct*	Struct;
-	FBuiltRange*	Range;
-};
-
-struct FBuiltMember
-{
-	FBuiltMember() = delete;
-	FBuiltMember(const FBuiltMember&) = delete;
-	FBuiltMember(FBuiltMember&& O);
-	FBuiltMember(FMemberId Name, FUnpackedLeafType Leaf, FOptionalEnumSchemaId Schema, uint64 Value);
-	//FBuiltMember(FMemberId Name, FEnumSchemaId Schema, ELeafWidth Width,  uint64 Value);
-	FBuiltMember(FMemberId Name, FTypedRange&& Range);
-	FBuiltMember(FMemberId Name, FStructSchemaId Schema, FBuiltStructPtr&& Value);
-	static FBuiltMember MakeSuper(FStructSchemaId Schema, FBuiltStructPtr&& Value);
-	~FBuiltMember(); // Deletes Value
-
-	FBuiltMember& operator=(const FBuiltMember&) = delete;
-	FBuiltMember& operator=(FBuiltMember&&);
-
-	FOptionalMemberId		Name;
-	FMemberSchema			Schema;
-	FBuiltValue				Value;
-
-private:
-	FBuiltMember(FOptionalMemberId N, FMemberSchema&& S, FBuiltValue V) : Name(N), Schema(MoveTemp(S)), Value(V) {}
-};
-
 struct FBuiltStruct
 {
-	~FBuiltStruct();
-
+	~FBuiltStruct() = delete; // Allocated in FScratchAllocator 
+	
 	uint16				NumMembers;
 	FBuiltMember		Members[0];
 };
 
 struct FBuiltRange
 {
-	~FBuiltRange() = delete;
-	[[nodiscard]] static FBuiltRange*					Create(uint64 NumItems, SIZE_T ItemSize);
-	static uint64										Delete(FBuiltRange* Range, FOptionalSchemaId InnerSchema, TConstArrayView<FMemberType> InnerTypes);
+	~FBuiltRange() = delete; // Allocated in FScratchAllocator
+
+	[[nodiscard]] static FBuiltRange*					Create(FScratchAllocator& Allocator, uint64 NumItems, SIZE_T ItemSize);
 
 	uint64												Num;
 	uint8												Data[0];
 	
-	TConstArrayView64<const FBuiltRange*>				AsRanges() const { return { reinterpret_cast<FBuiltRange const* const*>(Data), IntCastChecked<int64>(Num) }; }
-	TConstArrayView64<TUniquePtr<const FBuiltStruct>>	AsStructs() const { return { reinterpret_cast<const TUniquePtr<const FBuiltStruct>*>(Data), IntCastChecked<int64>(Num) }; }
+	TConstArrayView64<const FBuiltRange*>				AsRanges() const	{ return { reinterpret_cast<FBuiltRange const* const*>(Data),	static_cast<int64>(Num) }; }
+	TConstArrayView64<const FBuiltStruct*>				AsStructs() const	{ return { reinterpret_cast<FBuiltStruct const* const*>(Data),	static_cast<int64>(Num) }; }
 };
 
 //////////////////////////////////////////////////////////////////////////
diff --git a/Engine/Plugins/Experimental/PlainProps/Source/Private/PlainPropsSave.cpp b/Engine/Plugins/Experimental/PlainProps/Source/Private/PlainPropsSave.cpp
--- a/Engine/Plugins/Experimental/PlainProps/Source/Private/PlainPropsSave.cpp
+++ b/Engine/Plugins/Experimental/PlainProps/Source/Private/PlainPropsSave.cpp
@@ -53,21 +53,25 @@ struct FLeafRangeSaver
 {
 	FBuiltRange* Out;
 	uint8* OutIt;
+	uint8* OutEnd;
 
-	FLeafRangeSaver(uint64 Num, SIZE_T LeafSize)
-	: Out(FBuiltRange::Create(Num, LeafSize))
+	FLeafRangeSaver(FScratchAllocator& Scratch, uint64 Num, SIZE_T LeafSize)
+	: Out(FBuiltRange::Create(Scratch, Num, LeafSize))
 	, OutIt(Out->Data)
+	, OutEnd(OutIt + Num * LeafSize)
 	{}
 
 	void Append(FExistingItemSlice Slice, uint32 Stride, SIZE_T LeafSize, const FSaveContext&)
 	{
 		check(Stride == LeafSize);
+		check(OutIt + Slice.Num * LeafSize <= OutEnd);
 		FMemory::Memcpy(OutIt, Slice.Data, Slice.Num * LeafSize);
 		OutIt += Slice.Num * LeafSize;
 	}
 
 	FBuiltRange* Finish()
 	{
+		check(OutIt == OutEnd);
 		return Out;
 	}
 };
@@ -77,24 +81,27 @@ struct FLeafRangeSaver
 template<typename BuiltItemType, typename ItemSchemaType>
 struct TStructuralRangeSaver
 {
-	TArray64<BuiltItemType> Items;
+	FBuiltRange* Out;
+	BuiltItemType* It;
 
-	TStructuralRangeSaver(uint64 Num, ItemSchemaType)
-	{ 
-		Items.Reserve(Num);
-	}
+	TStructuralRangeSaver(FScratchAllocator& Scratch, uint64 Num, ItemSchemaType)
+	: Out(FBuiltRange::Create(Scratch, Num, sizeof(BuiltItemType)))
+	, It(reinterpret_cast<BuiltItemType*>(Out->Data))
+	{}
 
 	void Append(FExistingItemSlice Slice, uint32 Stride, ItemSchemaType Schema, const FSaveContext& OuterCtx)
 	{
+		check(It + Slice.Num <= reinterpret_cast<BuiltItemType*>(Out->Data) + Out->Num);
 		for (uint64 Idx = 0; Idx < Slice.Num; ++Idx)
 		{
-			Items.Emplace(SaveRangeItem(Slice.At(Idx, Stride), Schema, OuterCtx));
+			*It++ = SaveRangeItem(Slice.At(Idx, Stride), Schema, OuterCtx);
 		}
 	}
 
 	[[nodiscard]] FBuiltRange* Finish()
 	{
-		return Private::BuildStructuralRange(/* ownership xfer */ Items);
+		check(It == reinterpret_cast<BuiltItemType*>(Out->Data) + Out->Num);
+		return Out;
 	}
 };
 
@@ -112,7 +119,7 @@ template<class SaverType, typename InnerContextType>
 
 	if (const uint64 NumTotal = ReadCtx.Items.NumTotal)
 	{
-		SaverType Saver(NumTotal, InnerCtx);
+		SaverType Saver(OuterCtx.Scratch, NumTotal, InnerCtx);
 		while (true)
 		{
 			check(ReadCtx.Items.Slice.Num > 0);
@@ -136,7 +143,7 @@ template<class SaverType, typename InnerContextType>
 {
 	check(Member.NumRanges > 1);
 	check(Member.InnerTypes[0].IsRange());
-	return { Member.InnerTypes + 1, Member.RangeBindings + 1, Member.NumRanges - 1, Member.InnermostSchema };
+	return { Member.InnerTypes + 1, Member.RangeBindings + 1, static_cast<uint16>(Member.NumRanges - 1), Member.InnermostSchema };
 }
 
  ELeafWidth GetArithmeticWidth(FLeafBindType Leaf)
@@ -145,9 +152,9 @@ template<class SaverType, typename InnerContextType>
 	return Leaf.Arithmetic.Width;
  }
 
-[[nodiscard]] static FBuiltRange* SaveLeafRange(const uint8* Range, const ILeafRangeBinding& Binding, FUnpackedLeafType Leaf)
+[[nodiscard]] static FBuiltRange* SaveLeafRange(FScratchAllocator& Scratch, const uint8* Range, const ILeafRangeBinding& Binding, FUnpackedLeafType Leaf)
 {
-	FLeafRangeAllocator Allocator(Leaf);
+	FLeafRangeAllocator Allocator(Scratch, Leaf);
 	Binding.SaveLeaves(Range, Allocator);
 	return Allocator.GetAllocatedRange();
 }
@@ -171,7 +178,7 @@ template<class SaverType, typename InnerContextType>
 
 	if (Binding.IsLeafBinding())
 	{
-		return SaveLeafRange(Range, Binding.AsLeafBinding(), UnpackNonBitfield(InnerType.AsLeaf()));
+		return SaveLeafRange(Ctx.Scratch, Range, Binding.AsLeafBinding(), UnpackNonBitfield(InnerType.AsLeaf()));
 	}
 
 	const IItemRangeBinding& ItemBinding = Binding.AsItemBinding();
@@ -213,22 +220,25 @@ template<class SaverType, typename InnerContextType>
 	}
 }
 
-[[nodiscard]] static TArray<FMemberType> ToMemberTypes(const FMemberBindType* InnerTypes)
+[[nodiscard]] static FMemberType* CreateInnerRangeTypes(FScratchAllocator& Scratch, uint32 NumInnerTypes, const FMemberBindType* InnerTypes)
 {
-	TArray<FMemberType> Out;
-	for (const FMemberBindType* It = InnerTypes; true; ++It)
+	if (NumInnerTypes <= 1)
 	{
-		Out.Add(ToMemberType(*It));
-		if (!It->IsRange())
-		{
-			return Out;
-		}
-	} 
+		return nullptr;
+	}
+
+	FMemberType* Out = Scratch.AllocateArray<FMemberType>(NumInnerTypes);
+	for (uint32 Idx = 0; Idx < NumInnerTypes; ++Idx)
+	{
+		Out[Idx] = ToMemberType(InnerTypes[Idx]);
+	}
+	return Out;
 }
 
-[[nodiscard]] static FMemberSchema MakeSchema(FRangeMemberBinding Member)
+[[nodiscard]] static FMemberSchema CreateRangeSchema(FScratchAllocator& Scratch, FRangeMemberBinding Member)
 {
-	return { FMemberType(Member.RangeBindings[0].GetSizeType()), Member.InnermostSchema, ToMemberTypes(Member.InnerTypes) };
+	FMemberType* InnerRangeTypes = CreateInnerRangeTypes(Scratch, Member.NumRanges, Member.InnerTypes);
+	return { FMemberType(Member.RangeBindings[0].GetSizeType()), ToMemberType(Member.InnerTypes[0]), Member.NumRanges, Member.InnermostSchema, InnerRangeTypes };
 }
 
 static const uint8* At(const void* Ptr, SIZE_T Offset)
@@ -244,7 +254,7 @@ static void SaveMember(FMemberBuilder& Out, const void* Struct, FMemberId Name,
 
 static void SaveMember(FMemberBuilder& Out, const void* Struct, FMemberId Name, const FSaveContext& Ctx, FRangeMemberBinding Member)
 {
-	Out.AddRange(Name, { MakeSchema(Member), SaveRange(At(Struct, Member.Offset), Member, Ctx) });
+	Out.AddRange(Name, { CreateRangeSchema(Ctx.Scratch, Member), SaveRange(At(Struct, Member.Offset), Member, Ctx) });
 }
 
 static void SaveMember(FMemberBuilder& Out, const void* Struct, FMemberId Name, const FSaveContext& Ctx, FStructMemberBinding Member)
@@ -272,7 +282,7 @@ FBuiltStructPtr SaveStruct(const void* Struct, FStructSchemaId Id, const FSaveCo
 		}
 	}
 	
-	return Out.BuildAndReset(Declaration, Ctx.Declarations.GetDebug());
+	return Out.BuildAndReset(Ctx.Scratch, Declaration, Ctx.Declarations.GetDebug());
 }
 
 ////////////////////////////////////////////////////////////////////////////////////////////////
@@ -454,7 +464,7 @@ static void SaveMemberDelta(FMemberBuilder& Out, const void* Struct, const void*
 	const uint8* Range = At(Struct, Member.Offset);
 	if (DiffItem(Range, At(Default, Member.Offset), Ctx, Member))
 	{
-		Out.AddRange(Name, { MakeSchema(Member), SaveRange(Range, Member, Ctx) });
+		Out.AddRange(Name, { CreateRangeSchema(Ctx.Scratch, Member), SaveRange(Range, Member, Ctx) });
 	}
 }
 
@@ -489,7 +499,7 @@ FBuiltStructPtr SaveStructDelta(const void* Struct, const void* Default, FStruct
 		}
 	}
 	
-	return !Out.IsEmpty() ? Out.BuildAndReset(Declaration, Ctx.Declarations.GetDebug()) : nullptr;
+	return Out.IsEmpty() ? nullptr : Out.BuildAndReset(Ctx.Scratch, Declaration, Ctx.Declarations.GetDebug());
 }
 
 } // namespace PlainProps
\ No newline at end of file
diff --git a/Engine/Plugins/Experimental/PlainProps/Source/Private/PlainPropsWrite.cpp b/Engine/Plugins/Experimental/PlainProps/Source/Private/PlainPropsWrite.cpp
--- a/Engine/Plugins/Experimental/PlainProps/Source/Private/PlainPropsWrite.cpp
+++ b/Engine/Plugins/Experimental/PlainProps/Source/Private/PlainPropsWrite.cpp
@@ -255,7 +255,7 @@ static TArray<FMemberType> GetInnerRangeTypes(const FBuiltStructSchema& Struct)
 	TArray<FMemberType> Out;
 	for (const FMemberSchema* Schema : Struct.MemberSchemas)
 	{
-		Out.Append(Schema->InnerRangeTypes);
+		Out.Append(Schema->GetInnerRangeTypes());
 	}
 	return Out;
 }
@@ -634,8 +634,8 @@ private:
 		switch (Schema.Type.GetKind())
 		{
 		case EMemberKind::Leaf:		WriteLeaf(Schema.Type.AsLeaf(), Value.Leaf); break;
+		case EMemberKind::Range:	WriteRange(Schema.Type.AsRange().MaxSize, Schema.GetInnerRangeTypes(), Schema.InnerSchema, Value.Range); break;
 		case EMemberKind::Struct:	WriteStruct(Schema.Type.AsStruct(), static_cast<FStructSchemaId>(Schema.InnerSchema.Get()), *Value.Struct); break;
-		case EMemberKind::Range:	WriteRange(Schema.Type.AsRange().MaxSize, Schema.InnerRangeTypes, Schema.InnerSchema, Value.Range); break;
 		}
 	}
 
@@ -683,8 +683,8 @@ private:
 			switch (Types[0].GetKind())
 			{
 			case EMemberKind::Leaf:		WriteLeaves(Types[0].AsLeaf(), *Range);	break;
-			case EMemberKind::Struct:	WriteStructs(Types[0].AsStruct(), static_cast<FStructSchemaId>(InnermostSchema.Get()), Range->AsStructs()); break;
 			case EMemberKind::Range:	WriteRanges(Types[0].AsRange().MaxSize, Types.RightChop(1), InnermostSchema, Range->AsRanges()); break;
+			case EMemberKind::Struct:	WriteStructs(Types[0].AsStruct(), static_cast<FStructSchemaId>(InnermostSchema.Get()), Range->AsStructs()); break;
 			}
 		}
 	}
@@ -722,10 +722,10 @@ private:
 		Tmp.Reset();
 	}
 
-	void WriteStructs(FStructType StructType, FStructSchemaId Schema, TConstArrayView64<TUniquePtr<const FBuiltStruct>> Structs)
+	void WriteStructs(FStructType StructType, FStructSchemaId Schema, TConstArrayView64<const FBuiltStruct*> Structs)
 	{
 		Bytes.Append(reinterpret_cast<uint8*>(&Schema.Idx), StructType.IsDynamic * sizeof(Schema.Idx));
-		WriteSkippableItems(Structs, [=](FMemberWriter& Out, const TUniquePtr<const FBuiltStruct>& Struct) { Out.WriteStruct(StructType, Schema, *Struct); });
+		WriteSkippableItems(Structs, [=](FMemberWriter& Out, const FBuiltStruct* Struct) { Out.WriteStruct(StructType, Schema, *Struct); });
 	}
 	
 	void WriteRanges(ERangeSizeType NumType, TConstArrayView<FMemberType> Types, FOptionalSchemaId InnermostSchema, TConstArrayView64<const FBuiltRange*> Ranges)
diff --git a/Engine/Plugins/Experimental/PlainProps/Source/Public/PlainPropsBind.h b/Engine/Plugins/Experimental/PlainProps/Source/Public/PlainPropsBind.h
--- a/Engine/Plugins/Experimental/PlainProps/Source/Public/PlainPropsBind.h
+++ b/Engine/Plugins/Experimental/PlainProps/Source/Public/PlainPropsBind.h
@@ -19,6 +19,7 @@ class FIdIndexerBase;
 struct FLoadBatch;
 class FMemberBuilder;
 struct FSchemaBatch;
+class FScratchAllocator;
 class FStructBinding;
 class FRangeBinding;
 struct FTypedRange;
@@ -182,7 +183,7 @@ struct FRangeMemberBinding
 {
 	const FMemberBindType*	InnerTypes;
 	const FRangeBinding*	RangeBindings;
-	uint32					NumRanges; // At least 1, >1 for nested ranges
+	uint16					NumRanges; // At least 1, >1 for nested ranges
 	FOptionalSchemaId		InnermostSchema;
 	SIZE_T					Offset;
 };
@@ -418,13 +419,14 @@ public:
 // Possible save opt: Use paged linear allocator that only allocates on page exhaustion
 class FLeafRangeAllocator
 {
-	const FUnpackedLeafType		Expected;
+	FScratchAllocator&			Scratch;
 	FBuiltRange*				Range = nullptr;
+	const FUnpackedLeafType		Expected;
 
 	void* Allocate(FUnpackedLeafType Type, uint64 Num);
 
 public:
-	FLeafRangeAllocator(FUnpackedLeafType InExpected) : Expected(InExpected) {}
+	FLeafRangeAllocator(FScratchAllocator& InScratch, FUnpackedLeafType InExpected) : Scratch(InScratch), Expected(InExpected) {}
 
 	template<typename LeafType, typename SizeType>
 	LeafType* AllocateRange(SizeType Num)
diff --git a/Engine/Plugins/Experimental/PlainProps/Source/Public/PlainPropsBuild.h b/Engine/Plugins/Experimental/PlainProps/Source/Public/PlainPropsBuild.h
--- a/Engine/Plugins/Experimental/PlainProps/Source/Public/PlainPropsBuild.h
+++ b/Engine/Plugins/Experimental/PlainProps/Source/Public/PlainPropsBuild.h
@@ -4,6 +4,7 @@
 
 #include "PlainPropsDeclare.h"
 #include "PlainPropsTypes.h"
+#include "Algo/Compare.h"
 #include "Containers/Array.h"
 #include "Containers/ArrayView.h"
 #include "Memory/MemoryView.h"
@@ -18,26 +19,96 @@ struct FBuiltRange;
 class FDebugIds;
 struct FUnpackedLeafType;
 
-struct FBuiltStructDeleter 
+using FBuiltStructPtr = FBuiltStruct*;
+
+//////////////////////////////////////////////////////////////////////////
+
+/// Single-threaded scratch allocator for intermediate built representation
+class FScratchAllocator
 {
-	PLAINPROPS_API void operator()(FBuiltStruct* Ptr) const;
+	struct FPage
+	{
+		FPage*			PrevPage;
+		uint8			Data[0];
+	};
+
+	static constexpr uint32 PageSize = 65536;
+	static constexpr uint32 DataSize = PageSize - offsetof(FPage, Data);
+	
+	uint8*					Cursor = nullptr;
+	uint8*					PageEnd = nullptr;
+	FPage*					LastPage = nullptr;
+
+	PLAINPROPS_API uint8* AllocateInNewPage(SIZE_T Size, uint32 Alignment);
+
+public:
+	UE_NONCOPYABLE(FScratchAllocator);
+	FScratchAllocator() = default;
+	PLAINPROPS_API ~FScratchAllocator();
+
+	inline void* Allocate(SIZE_T Size, uint32 Alignment)
+	{
+		uint8* Out = Align(Cursor, Alignment);
+		if (Out + Size <= PageEnd)
+		{
+			Cursor = Out + Size;
+			return Out;
+		}
+
+		return AllocateInNewPage(Size, Alignment);
+	}
+
+	inline void* AllocateZeroed(SIZE_T Size, uint32 Alignment)
+	{
+		void* Out = Allocate(Size, Alignment);
+		FMemory::Memzero(Out, Size);
+		return Out;
+	}
+
+	template<typename T>
+	inline T* AllocateArray(uint64 Num)
+	{
+		T* Out = static_cast<T*>(Allocate(Num * sizeof(T), alignof(T)));
+		for (uint64 Idx = 0; Idx < Num; ++Idx)
+		{
+			new (Out + Idx) T;
+		}
+		return Out;
+	}
 };
-using FBuiltStructPtr = TUniquePtr<FBuiltStruct, FBuiltStructDeleter>;
 
 //////////////////////////////////////////////////////////////////////////
 
 struct FMemberSchema
 {
-	FMemberType				Type;
-	FOptionalSchemaId		InnerSchema;
-	TArray<FMemberType>		InnerRangeTypes;
+	FMemberType						Type;
+	FMemberType						InnerRangeType;
+	uint16							NumInnerRanges;
+	FOptionalSchemaId				InnerSchema;
+	const FMemberType*				NestedRangeTypes;
+
+	TConstArrayView<FMemberType> GetInnerRangeTypes() const
+	{
+		return MakeArrayView(NestedRangeTypes ? NestedRangeTypes : &InnerRangeType, NumInnerRanges);
+	}
+
+	FMemberType GetInnermostType() const
+	{
+		return NumInnerRanges ? GetInnerRangeTypes().Last() : Type;
+	}
 
-	FMemberType GetInnermostType() const { return Type.IsRange() ? InnerRangeTypes.Last() : Type; }
+	[[nodiscard]] PLAINPROPS_API FMemberType& EditInnermostType(FScratchAllocator& Scratch);
+
+	void CheckInvariants()
+	{
+		check(Type.IsRange() == !!NumInnerRanges);
+		check(!!NestedRangeTypes == (NumInnerRanges > 1));
+	}
 };
 
-inline bool operator==(const FMemberSchema& A, const FMemberSchema& B)
+inline bool operator==(FMemberSchema A, FMemberSchema B)
 {
-	return A.Type == B.Type && A.InnerSchema == B.InnerSchema && A.InnerRangeTypes == B.InnerRangeTypes;
+	return A.Type == B.Type && A.InnerSchema == B.InnerSchema && Algo::Compare(A.GetInnerRangeTypes(), B.GetInnerRangeTypes());
 }
 //////////////////////////////////////////////////////////////////////////
 
@@ -65,19 +136,18 @@ struct FTypedRange
 	FBuiltRange* Values = nullptr;
 };
 
-
 template<typename LeafType, typename SizeType>
 FMemberSchema MakeLeafRangeSchema()
 {
 	check(std::is_arithmetic_v<LeafType>);
-	return { FMemberType(RangeSizeOf(SizeType{})), NoId, { ReflectLeaf<LeafType>.Pack() } };
+	return { FMemberType(RangeSizeOf(SizeType{})), ReflectLeaf<LeafType>.Pack(), 1, NoId, nullptr };
 }
 
 template<typename EnumType, typename SizeType>
 FMemberSchema MakeEnumRangeSchema(FEnumSchemaId Schema)
 {
 	check(std::is_enum_v<EnumType>);
-	return { FMemberType(RangeSizeOf(SizeType{})), Schema, { ReflectLeaf<EnumType>.Pack() } };
+	return { FMemberType(RangeSizeOf(SizeType{})), ReflectLeaf<EnumType>.Pack(), 1, Schema, nullptr };
 }
 
 inline constexpr FMemberType DefaultStructType = FMemberType(FStructType{EMemberKind::Struct, /* IsDynamic */ 0, /* IsSuper */ 0});
@@ -86,65 +156,70 @@ inline constexpr FMemberType SuperStructType =	 FMemberType(FStructType{EMemberK
 
 inline FMemberSchema MakeStructRangeSchema(ERangeSizeType SizeType, FStructSchemaId Schema)
 {
-	return { FMemberType(SizeType), Schema, { DefaultStructType } };
+	return { FMemberType(SizeType), DefaultStructType, 1, Schema, nullptr };
 }
 
-PLAINPROPS_API FMemberSchema MakeNestedRangeSchema(ERangeSizeType SizeType, const FMemberSchema& InnerRangeSchema);
+PLAINPROPS_API FMemberSchema MakeNestedRangeSchema(FScratchAllocator& Scratch, ERangeSizeType SizeType, FMemberSchema InnerRangeSchema);
 
-namespace Private
-{
-	[[nodiscard]] PLAINPROPS_API FBuiltRange*	BuildStructuralRange(/* in-out */ TArrayView64<FBuiltRange*> Structs);
-	[[nodiscard]] PLAINPROPS_API FBuiltRange*	BuildStructuralRange(/* in-out */ TArrayView64<FBuiltStructPtr> Structs);
-	[[nodiscard]] PLAINPROPS_API FBuiltRange*	BuildLeafRange(FUnpackedLeafType Leaf, uint64 Num, FMemoryView Values);
-	PLAINPROPS_API void							NormalizeLeafRange(FUnpackedLeafType Leaf, FBuiltRange& Out);
+//////////////////////////////////////////////////////////////////////////
 
-	template<typename LeafType, typename SizeType>
-	[[nodiscard]] FBuiltRange* BuildLeafRange(const LeafType* Values, SizeType InNum)
-	{
-		static_assert(SizeOf(ReflectLeaf<LeafType>.Width) == sizeof(LeafType));
-
-		uint64 Num = IntCastChecked<uint64>(InNum);
-		FBuiltRange* Range = Num > 0 ? Private::BuildLeafRange(ReflectLeaf<LeafType>, Num, MakeMemoryView(Values, Num * sizeof(LeafType))) : nullptr;
-	
-		if (Range && std::is_floating_point_v<LeafType>)
-		{
-			NormalizeLeafRange(ReflectLeaf<LeafType>, *Range);
-		}
-	
-		return Range;
-	}
-}
+[[nodiscard]] PLAINPROPS_API FBuiltRange* CloneLeaves(FScratchAllocator& Scratch, uint64 Num, const void* Data, SIZE_T LeafSize);
 
 template<typename LeafType, typename SizeType>
-[[nodiscard]] FTypedRange BuildLeafRange(const LeafType* Values, SizeType Num)
+[[nodiscard]] FTypedRange BuildLeafRange(FScratchAllocator& Scratch, const LeafType* Values, SizeType Num)
 {
-	return { MakeLeafRangeSchema<LeafType, SizeType>(), Private::BuildLeafRange(Values, Num) };
+	// todo: detect invalid floats
+	return { MakeLeafRangeSchema<LeafType, SizeType>(), CloneLeaves(Scratch, Num, Values, sizeof(LeafType)) };
 }
 
 template<typename LeafType, typename SizeType>
-[[nodiscard]] FTypedRange BuildLeafRange(TConstArrayView<LeafType, SizeType> Values)
+[[nodiscard]] FTypedRange BuildLeafRange(FScratchAllocator& Scratch, TConstArrayView<LeafType, SizeType> Values)
 {
-	return BuildLeafRange(Values.GetData(), Values.Num());
+	// todo: detect invalid floats
+	return { MakeLeafRangeSchema<LeafType, SizeType>(), CloneLeaves(Scratch, Values.Num(), Values.GetData(), sizeof(LeafType)) };
 }
 
 template<typename EnumType, typename SizeType>
-[[nodiscard]] FTypedRange BuildEnumRange(FEnumSchemaId Enum, TConstArrayView<EnumType, SizeType> Values)
+[[nodiscard]] FTypedRange BuildEnumRange(FScratchAllocator& Scratch, FEnumSchemaId Enum, TConstArrayView<EnumType, SizeType> Values)
 {
-	return { MakeEnumRangeSchema<EnumType, SizeType>(Enum), Private::BuildLeafRange(Values.GetData(), Values.Num()) };
+	return { MakeEnumRangeSchema<EnumType, SizeType>(Enum), CloneLeaves(Scratch, Values.Num(), Values.GetData(), sizeof(EnumType)) };
 }
 
-[[nodiscard]] inline FTypedRange BuildStructRange(FStructSchemaId Schema, ERangeSizeType SizeType, /* in-out */ TArrayView64<FBuiltStructPtr> Values )
+[[nodiscard]] inline FTypedRange MakeStructRange(FStructSchemaId Schema, ERangeSizeType SizeType, FBuiltRange* Values )
 {
-	return { MakeStructRangeSchema(SizeType, Schema), Values.Num() ? Private::BuildStructuralRange(/* ownership xfer */ Values) : nullptr };
+	return { MakeStructRangeSchema(SizeType, Schema), Values };
 }
 
+//////////////////////////////////////////////////////////////////////////
+
+union FBuiltValue
+{
+	uint64			Leaf;
+	FBuiltStruct*	Struct;
+	FBuiltRange*	Range;
+};
+
+struct FBuiltMember
+{
+	FBuiltMember(FMemberId Name, FUnpackedLeafType Leaf, FOptionalEnumSchemaId Schema, uint64 Value);
+	FBuiltMember(FMemberId Name, FTypedRange Range);
+	FBuiltMember(FMemberId Name, FStructSchemaId Schema, FBuiltStructPtr Value);
+	static FBuiltMember MakeSuper(FStructSchemaId Schema, FBuiltStructPtr Value);
+
+	FOptionalMemberId		Name;
+	FMemberSchema			Schema;
+	FBuiltValue				Value;
+
+private:
+	FBuiltMember(FOptionalMemberId N, FMemberSchema&& S, FBuiltValue V) : Name(N), Schema(MoveTemp(S)), Value(V) {}
+};
+
+//////////////////////////////////////////////////////////////////////////
+
 // Builds an ordered list of properties to be saved
 class FMemberBuilder
 {
 public:
-	FMemberBuilder();
-	~FMemberBuilder();
-
 	template<typename LeafType>
 	void Add(FMemberId Name, LeafType Value)
 	{
@@ -163,19 +238,21 @@ public:
 	void AddEnum32(FMemberId Name, FEnumSchemaId Schema, uint32 Value)	{ AddLeaf(Name, {ELeafType::Enum, ELeafWidth::B32}, ToOptional(Schema), Value); }
 	void AddEnum64(FMemberId Name, FEnumSchemaId Schema, uint64 Value)	{ AddLeaf(Name, {ELeafType::Enum, ELeafWidth::B64}, ToOptional(Schema), Value); }
 
-	PLAINPROPS_API void AddLeaf(FMemberId Name, FUnpackedLeafType Leaf, FOptionalEnumSchemaId Enum, uint64 Value);
-	PLAINPROPS_API void AddStruct(FMemberId Name, FStructSchemaId Schema, FBuiltStructPtr&& Struct);	
-	PLAINPROPS_API void AddRange(FMemberId Name, FTypedRange&& Range);
+	void AddLeaf(FMemberId Name, FUnpackedLeafType Leaf, FOptionalEnumSchemaId Enum, uint64 Value)	{ Members.Emplace(Name, Leaf, Enum, Value); }
+	void AddRange(FMemberId Name, FTypedRange Range)												{ Members.Emplace(Name, Range); }
+	void AddStruct(FMemberId Name, FStructSchemaId Schema, FBuiltStructPtr Struct)					{ Members.Emplace(Name, Schema, Struct); }
 	
 	// Build members into a single nested super struct member, no-op if no non-super members has been added
-	PLAINPROPS_API void BuildSuperStruct(const FStructDeclaration& Super, const FDebugIds& Debug);
+	PLAINPROPS_API void BuildSuperStruct(FScratchAllocator&	Scratch, const FStructDeclaration& Super, const FDebugIds& Debug);
 
-	[[nodiscard]] PLAINPROPS_API FBuiltStructPtr BuildAndReset(const FStructDeclaration& Declared, const FDebugIds& Debug);
+	[[nodiscard]] PLAINPROPS_API FBuiltStructPtr BuildAndReset(FScratchAllocator& Scratch, const FStructDeclaration& Declared, const FDebugIds& Debug);
 
 	bool IsEmpty() const { return Members.IsEmpty(); }
 
 private:
-	TArray<FBuiltMember>	Members;
+	using FBuiltMemberArray = TArray<FBuiltMember, TInlineAllocator<16>>;
+	
+	FBuiltMemberArray		Members;
 	
 //	PLAINPROPS_API uint8* AllocateLeafRange(FMemberId Name, FUnpackedLeafType Leaf, ERangeSizeType RangeMax, uint64 Num);
 	
@@ -203,7 +280,7 @@ public:
 	 
 	FMemberBuilder& operator[](uint64 Idx) { return Structs[Idx]; }
 
-	FTypedRange BuildAndReset(const FStructDeclaration& Declared, const FDebugIds& Debug);
+	FTypedRange BuildAndReset(FScratchAllocator& Scratch, const FStructDeclaration& Declared, const FDebugIds& Debug);
 
 private:
 	TArray64<FMemberBuilder> Structs; 
@@ -229,7 +306,7 @@ public:
 		Ranges.Add(Range.Values);
 	}
 
-	[[nodiscard]] FTypedRange BuildAndReset(ERangeSizeType SizeType);
+	[[nodiscard]] FTypedRange BuildAndReset(FScratchAllocator& Scratch, ERangeSizeType SizeType);
 
 private:
 	TArray64<FBuiltRange*> Ranges; 
diff --git a/Engine/Plugins/Experimental/PlainProps/Source/Public/PlainPropsBuildSchema.h b/Engine/Plugins/Experimental/PlainProps/Source/Public/PlainPropsBuildSchema.h
--- a/Engine/Plugins/Experimental/PlainProps/Source/Public/PlainPropsBuildSchema.h
+++ b/Engine/Plugins/Experimental/PlainProps/Source/Public/PlainPropsBuildSchema.h
@@ -46,8 +46,8 @@ public:
 	using FStructDeclarations = TConstArrayView<TUniquePtr<FStructDeclaration>>;
 	using FEnumDeclarations = TConstArrayView<TUniquePtr<FEnumDeclaration>>;
 
-	PLAINPROPS_API explicit FSchemasBuilder(const FDeclarations& Declarations);
-	PLAINPROPS_API FSchemasBuilder(FStructDeclarations InStructs, FEnumDeclarations InEnums, const FDebugIds& InDebug);
+	PLAINPROPS_API explicit FSchemasBuilder(const FDeclarations& Declarations, FScratchAllocator& Scratch);
+	PLAINPROPS_API FSchemasBuilder(FStructDeclarations InStructs, FEnumDeclarations InEnums, const FDebugIds& InDebug, FScratchAllocator& Scratch);
 	PLAINPROPS_API ~FSchemasBuilder();
 
 	PLAINPROPS_API FEnumSchemaBuilder&			NoteEnum(FEnumSchemaId Id);
@@ -62,6 +62,7 @@ private:
 	TArray<int32>								EnumIndices;
 	TPagedArray<FStructSchemaBuilder, 4096>		Structs;		// TPagedArray for stable references
 	TPagedArray<FEnumSchemaBuilder, 4096>		Enums;			// TPagedArray for stable references
+	FScratchAllocator&							Scratch;
 	const FDebugIds&							Debug;
 	bool										bBuilt = false;
 
diff --git a/Engine/Plugins/Experimental/PlainProps/Source/Public/PlainPropsSave.h b/Engine/Plugins/Experimental/PlainProps/Source/Public/PlainPropsSave.h
--- a/Engine/Plugins/Experimental/PlainProps/Source/Public/PlainPropsSave.h
+++ b/Engine/Plugins/Experimental/PlainProps/Source/Public/PlainPropsSave.h
@@ -23,13 +23,14 @@ struct FSaveContext
 	const FDeclarations&		Declarations;
 	const FSchemaBindings&		Schemas;
 	FCustomBindings&			Customs;
+	FScratchAllocator&			Scratch;
 	FDefaultStructs				Defaults;
 };
 
 template<typename Runtime>
-FSaveContext MakeSaveContext(FDefaultStructs Defaults)
+FSaveContext MakeSaveContext(FDefaultStructs Defaults, FScratchAllocator& Scratch)
 {
-	return { Runtime::GetTypes(), Runtime::GetSchemas(), Runtime::GetCustoms(), Defaults };
+	return { Runtime::GetTypes(), Runtime::GetSchemas(), Runtime::GetCustoms(), Scratch, Defaults };
 }
 
 [[nodiscard]] PLAINPROPS_API FBuiltStructPtr SaveStruct(const void* Struct, FStructSchemaId Id, const FSaveContext& Context);
diff --git a/Engine/Plugins/Experimental/PlainProps/Source/Tests/PlainPropsTests.cpp b/Engine/Plugins/Experimental/PlainProps/Source/Tests/PlainPropsTests.cpp
--- a/Engine/Plugins/Experimental/PlainProps/Source/Tests/PlainPropsTests.cpp
+++ b/Engine/Plugins/Experimental/PlainProps/Source/Tests/PlainPropsTests.cpp
@@ -107,7 +107,7 @@ inline constexpr uint32 TestMagics[] = {0xFEEDF00D, 0xABCD1234, 0xDADADAAA, 0x99
 class FTestBatchBuilder : public TIdIndexer<FAnsiString>
 {
 public:
-	FTestBatchBuilder() : Declarations(*this) {}
+	FTestBatchBuilder(FScratchAllocator& InScratch) : Declarations(*this), Scratch(InScratch) {}
 
 	FEnumSchemaId				DeclareEnum(FTypeId Type, EEnumMode Mode, ELeafWidth Width, std::initializer_list<const char*> Names, std::initializer_list<uint64> Constants);
 	FEnumSchemaId				DeclareEnum(const char* Scope, const char* Name, EEnumMode Mode, ELeafWidth Width, std::initializer_list<const char*> Names, std::initializer_list<uint64> Constants);
@@ -124,7 +124,8 @@ public:
 
 private:
 	TArray<TPair<FStructSchemaId, FBuiltStructPtr>>	Objects;
-	FDeclarations												Declarations;
+	FDeclarations									Declarations;
+	FScratchAllocator&								Scratch;
 
 	TArray<FMemberId> NameMembers(std::initializer_list<const char*> Members)
 	{
@@ -178,13 +179,13 @@ FEnumSchemaId FTestBatchBuilder::DeclareEnum(const char* Scope, const char* Name
 
 void FTestBatchBuilder::AddObject(FStructSchemaId Schema, FMemberBuilder&& Members)
 {
-	Objects.Emplace(Schema, Members.BuildAndReset(Declarations.Get(Schema), *this));
+	Objects.Emplace(Schema, Members.BuildAndReset(Scratch, Declarations.Get(Schema), *this));
 }
 
 TArray64<uint8> FTestBatchBuilder::Write()
 {
 	// Build partial schemas
-	FSchemasBuilder SchemaBuilders(Declarations.GetStructs(), Declarations.GetEnums(), *this);
+	FSchemasBuilder SchemaBuilders(Declarations, Scratch);
 	for (const TPair<FStructSchemaId, FBuiltStructPtr>& Object : Objects)
 	{
 		SchemaBuilders.NoteStructAndMembers(Object.Key, *Object.Value);
@@ -325,12 +326,13 @@ private:
 	TArray<FStructView> Objects;
 };
 
-static void TestSerialize(void (*BuildObjects)(FTestBatchBuilder&), void (*CheckObjects)(TConstArrayView<FStructView>, const FTestNameReader&))
+static void TestSerialize(void (*BuildObjects)(FTestBatchBuilder&, FScratchAllocator&), void (*CheckObjects)(TConstArrayView<FStructView>, const FTestNameReader&))
 {
 	TArray64<uint8> Data;
 	{
-		FTestBatchBuilder Batch;
-		BuildObjects(Batch);
+		FScratchAllocator Scratch;
+		FTestBatchBuilder Batch(Scratch);
+		BuildObjects(Batch, Scratch);
 		Data = Batch.Write();
 	}
 
@@ -375,7 +377,7 @@ TEST_CASE_NAMED(FPlainPropsReadWriteTest, "System::Core::Serialization::PlainPro
 	{
 		static constexpr std::initializer_list<const char*> MemberNames =  {"b0", "b1", "b2", "b3", "b4", "b5", "b6", "b7", "b8", "b9", "b10", "b11"};
 
-		TestSerialize([](FTestBatchBuilder& Batch)
+		TestSerialize([](FTestBatchBuilder& Batch, FScratchAllocator& Scratch)
 		{
 			FStructSchemaId SchemaId = Batch.DeclareStruct("Testing", "Bools", MemberNames, EMemberPresence::AllowSparse);
 
@@ -468,7 +470,7 @@ TEST_CASE_NAMED(FPlainPropsReadWriteTest, "System::Core::Serialization::PlainPro
 	{
 		static constexpr std::initializer_list<const char*> MemberNames =  {"F32", "F64", "S8", "U8", "S16", "U16", "S32", "U32", "S64", "U64"};
 
-		TestSerialize([](FTestBatchBuilder& Batch)
+		TestSerialize([](FTestBatchBuilder& Batch, FScratchAllocator& Scratch)
 		{
 			FStructSchemaId SchemaId = Batch.DeclareStruct("Test", "Numbers", MemberNames, EMemberPresence::AllowSparse);
 
@@ -574,7 +576,7 @@ TEST_CASE_NAMED(FPlainPropsReadWriteTest, "System::Core::Serialization::PlainPro
 	{
 		static constexpr std::initializer_list<const char*> MemberNames =  {"A", "B", "C"};
 
-		TestSerialize([](FTestBatchBuilder& Batch)
+		TestSerialize([](FTestBatchBuilder& Batch, FScratchAllocator& Scratch)
 		{
 			FStructSchemaId ExplicitId = Batch.DeclareStruct("Test", "ExplicitDense", MemberNames, EMemberPresence::RequireAll);
 			FStructSchemaId ImplicitId = Batch.DeclareStruct("Test", "ImplicitDense", MemberNames, EMemberPresence::AllowSparse);
@@ -620,7 +622,7 @@ TEST_CASE_NAMED(FPlainPropsReadWriteTest, "System::Core::Serialization::PlainPro
 		static constexpr std::initializer_list<const char*> StructMembers =  {"Nested", "Leaf"};
 		static constexpr std::initializer_list<const char*> NestedMembers =  {"I1", "I2"};
 
-		TestSerialize([](FTestBatchBuilder& Batch)
+		TestSerialize([](FTestBatchBuilder& Batch, FScratchAllocator& Scratch)
 		{
 			FStructSchemaId ObjectId = Batch.DeclareStruct("Test", "Object", ObjectMembers, EMemberPresence::AllowSparse);
 			FStructSchemaId StructId = Batch.DeclareStruct("Test", "Struct", StructMembers, EMemberPresence::AllowSparse);
@@ -628,14 +630,14 @@ TEST_CASE_NAMED(FPlainPropsReadWriteTest, "System::Core::Serialization::PlainPro
 
 			FMemberBuilder Members;
 			Members.Add(	Batch.NameMember("I1"), 100);
-			FBuiltStructPtr NestedInStruct = Members.BuildAndReset(Batch.Get(NestedId), Batch);
+			FBuiltStructPtr NestedInStruct = Members.BuildAndReset(Scratch, Batch.Get(NestedId), Batch);
 
 			Members.AddStruct(Batch.NameMember("Nested"), NestedId, MoveTemp(NestedInStruct));
 			Members.Add(Batch.NameMember("Leaf"), true);
-			FBuiltStructPtr Struct = Members.BuildAndReset(Batch.Get(StructId), Batch);
+			FBuiltStructPtr Struct = Members.BuildAndReset(Scratch, Batch.Get(StructId), Batch);
 		
 			Members.Add(Batch.NameMember("I2"), 200);
-			FBuiltStructPtr NestedInObject = Members.BuildAndReset(Batch.Get(NestedId), Batch);
+			FBuiltStructPtr NestedInObject = Members.BuildAndReset(Scratch, Batch.Get(NestedId), Batch);
 
 			Members.Add(Batch.NameMember("L1"), 123.f);
 			Members.AddStruct(Batch.NameMember("S"), StructId, MoveTemp(Struct));
@@ -668,7 +670,7 @@ TEST_CASE_NAMED(FPlainPropsReadWriteTest, "System::Core::Serialization::PlainPro
 		static constexpr std::initializer_list<const char*> MemberNames = 
 		{ "A2", "A0", "B0", "B4", "B5", "B7", "C3", "D34", "Max8", "Max16", "Max32", "Max64", "IF" };
 
-		TestSerialize([](FTestBatchBuilder& Batch)
+		TestSerialize([](FTestBatchBuilder& Batch, FScratchAllocator& Scratch)
 		{
 			// Test create holes in the original FNameId, FStructSchemaId and FEnumSchemaId index range
 			FStructSchemaId UnusedId = Batch.DeclareStruct("Test", "UnusedStruct", {"U1", "U2"}, EMemberPresence::AllowSparse);
@@ -748,7 +750,7 @@ TEST_CASE_NAMED(FPlainPropsReadWriteTest, "System::Core::Serialization::PlainPro
 		enum class EUnused1 : uint8 { X };
 		enum class EUnused2 : uint8 { Y };
 
-		TestSerialize([](FTestBatchBuilder& Batch)
+		TestSerialize([](FTestBatchBuilder& Batch, FScratchAllocator& Scratch)
 		{
 			static constexpr std::initializer_list<const char*> MemberNames = { "B0", "B1", "B8", "B9", "D0", "D3", "Hi", "E3", "E0" };
 		
@@ -758,15 +760,15 @@ TEST_CASE_NAMED(FPlainPropsReadWriteTest, "System::Core::Serialization::PlainPro
 			FEnumSchemaId UnusedEnum2	= Batch.DeclareEnum("Test", "Unused2", EEnumMode::Flat, ELeafWidth::B8, {"Y"}, {0});
 		
 			FMemberBuilder Members;
-			Members.AddRange(Batch.NameMember("B0"), BuildLeafRange(TConstArrayView<bool>()));
-			Members.AddRange(Batch.NameMember("B1"), BuildLeafRange(MakeArrayView({true})));
-			Members.AddRange(Batch.NameMember("B8"), BuildLeafRange(MakeArrayView({false, true, false, true, false, true, false, true})));
-			Members.AddRange(Batch.NameMember("B9"), BuildLeafRange(MakeArrayView({true, false, true, false, true, false, true, false, true})));
-			Members.AddRange(Batch.NameMember("D0"), BuildLeafRange(TConstArrayView<double>()));
-			Members.AddRange(Batch.NameMember("D3"), BuildLeafRange(MakeArrayView({DBL_MIN, 0.0, DBL_MAX})));
-			Members.AddRange(Batch.NameMember("Hi"), BuildLeafRange(MakeArrayView(u8"Hello!")));
-			Members.AddRange(Batch.NameMember("E3"), BuildEnumRange(Enum, MakeArrayView({EABCD::B, EABCD::A, EABCD::D})));
-			Members.AddRange(Batch.NameMember("E0"), BuildEnumRange(Enum, TConstArrayView<EUnused1>()));
+			Members.AddRange(Batch.NameMember("B0"), BuildLeafRange(Scratch, TConstArrayView<bool>()));
+			Members.AddRange(Batch.NameMember("B1"), BuildLeafRange(Scratch, MakeArrayView({true})));
+			Members.AddRange(Batch.NameMember("B8"), BuildLeafRange(Scratch, MakeArrayView({false, true, false, true, false, true, false, true})));
+			Members.AddRange(Batch.NameMember("B9"), BuildLeafRange(Scratch, MakeArrayView({true, false, true, false, true, false, true, false, true})));
+			Members.AddRange(Batch.NameMember("D0"), BuildLeafRange(Scratch, TConstArrayView<double>()));
+			Members.AddRange(Batch.NameMember("D3"), BuildLeafRange(Scratch, MakeArrayView({DBL_MIN, 0.0, DBL_MAX})));
+			Members.AddRange(Batch.NameMember("Hi"), BuildLeafRange(Scratch, MakeArrayView(u8"Hello!")));
+			Members.AddRange(Batch.NameMember("E3"), BuildEnumRange(Scratch, Enum, MakeArrayView({EABCD::B, EABCD::A, EABCD::D})));
+			Members.AddRange(Batch.NameMember("E0"), BuildEnumRange(Scratch, Enum, TConstArrayView<EUnused1>()));
 
 			Batch.AddObject(ObjectId, MoveTemp(Members));
 		}, 
@@ -799,7 +801,7 @@ TEST_CASE_NAMED(FPlainPropsReadWriteTest, "System::Core::Serialization::PlainPro
 		
 	SECTION("StructRange")
 	{
-		TestSerialize([](FTestBatchBuilder& Batch)
+		TestSerialize([](FTestBatchBuilder& Batch, FScratchAllocator& Scratch)
 		{
 			FStructSchemaId ObjectId = Batch.DeclareStruct("Test", "Object", {"Structs"}, EMemberPresence::AllowSparse);
 			FStructSchemaId StructId = Batch.DeclareStruct("Test", "Struct", {"I", "F"}, EMemberPresence::AllowSparse);
@@ -809,7 +811,7 @@ TEST_CASE_NAMED(FPlainPropsReadWriteTest, "System::Core::Serialization::PlainPro
 			Structs[1].Add(Batch.NameMember("F"), 1.f);
 		
 			FMemberBuilder Members;
-			Members.AddRange(Batch.NameMember("Structs"), Structs.BuildAndReset(Batch.Get(StructId), Batch));
+			Members.AddRange(Batch.NameMember("Structs"), Structs.BuildAndReset(Scratch, Batch.Get(StructId), Batch));
 
 			Batch.AddObject(ObjectId, MoveTemp(Members));
 		}, 
@@ -829,7 +831,7 @@ TEST_CASE_NAMED(FPlainPropsReadWriteTest, "System::Core::Serialization::PlainPro
 	{
 		enum class EAB : uint8 { A = 1, B = 4 };
 
-		TestSerialize([](FTestBatchBuilder& Batch)
+		TestSerialize([](FTestBatchBuilder& Batch, FScratchAllocator& Scratch)
 		{
 			FStructSchemaId Object = Batch.DeclareStruct("Test", "Object", {"IntRs", "EmptyRs", "EnumRs", "StructRs", "StructRRs"}, EMemberPresence::AllowSparse);
 			FStructSchemaId XY = Batch.DeclareStruct("Test", "XY", {"X", "Y"}, EMemberPresence::RequireAll);
@@ -837,13 +839,13 @@ TEST_CASE_NAMED(FPlainPropsReadWriteTest, "System::Core::Serialization::PlainPro
 			FEnumSchemaId Enum	= Batch.DeclareEnum("Test", "AB", EEnumMode::Flag, ELeafWidth::B8, {"A", "B"}, {1, 4});
 	
 			FNestedRangeBuilder IntRs(MakeLeafRangeSchema<int32, int32>(), 3);
-			IntRs.Add(BuildLeafRange(MakeArrayView({1})));
+			IntRs.Add(BuildLeafRange(Scratch, MakeArrayView({1})));
 			IntRs.Add({});
-			IntRs.Add(BuildLeafRange(MakeArrayView({2, 3})));
+			IntRs.Add(BuildLeafRange(Scratch, MakeArrayView({2, 3})));
 
 			FNestedRangeBuilder EnumRs(MakeEnumRangeSchema<EAB, int32>(Enum), 2);
 			EnumRs.Add({});
-			EnumRs.Add(BuildEnumRange(Enum, MakeArrayView({EAB::A, EAB(0), EAB::B})));
+			EnumRs.Add(BuildEnumRange(Scratch, Enum, MakeArrayView({EAB::A, EAB(0), EAB::B})));
 
 			FStructRangeBuilder XYs(uint64(2));
 			XYs[0].Add(Batch.NameMember("X"), 1.f);
@@ -851,7 +853,7 @@ TEST_CASE_NAMED(FPlainPropsReadWriteTest, "System::Core::Serialization::PlainPro
 			XYs[1].Add(Batch.NameMember("X"), 3.f);
 			XYs[1].Add(Batch.NameMember("Y"), 4.f);
 			FNestedRangeBuilder StructRs(MakeStructRangeSchema(ERangeSizeType::U64, XY), 1);
-			StructRs.Add(XYs.BuildAndReset(Batch.Get(XY), Batch));
+			StructRs.Add(XYs.BuildAndReset(Scratch, Batch.Get(XY), Batch));
 
 			FStructRangeBuilder ZWs(int16(3));
 			ZWs[0].Add(Batch.NameMember("Z"), 1.5f);
@@ -859,17 +861,17 @@ TEST_CASE_NAMED(FPlainPropsReadWriteTest, "System::Core::Serialization::PlainPro
 			ZWs[2].Add(Batch.NameMember("W"), 3.5f);
 			FMemberSchema ZWRangeSchema = MakeStructRangeSchema(ERangeSizeType::S16, ZW);
 			FNestedRangeBuilder ZWRs(ZWRangeSchema, 1);
-			ZWRs.Add(ZWs.BuildAndReset(Batch.Get(ZW), Batch));
-			FNestedRangeBuilder StructRRs(MakeNestedRangeSchema(ERangeSizeType::U32, ZWRangeSchema), 1);
-			StructRRs.Add(ZWRs.BuildAndReset(ERangeSizeType::U32));
+			ZWRs.Add(ZWs.BuildAndReset(Scratch, Batch.Get(ZW), Batch));
+			FNestedRangeBuilder StructRRs(MakeNestedRangeSchema(Scratch, ERangeSizeType::U32, ZWRangeSchema), 1);
+			StructRRs.Add(ZWRs.BuildAndReset(Scratch, ERangeSizeType::U32));
 
 
 			FMemberBuilder Members;
-			Members.AddRange(Batch.NameMember("IntRs"), IntRs.BuildAndReset(ERangeSizeType::S32));
-			Members.AddRange(Batch.NameMember("EmptyRs"), IntRs.BuildAndReset(ERangeSizeType::S32));
-			Members.AddRange(Batch.NameMember("EnumRs"), EnumRs.BuildAndReset(ERangeSizeType::U8));
-			Members.AddRange(Batch.NameMember("StructRs"), StructRs.BuildAndReset(ERangeSizeType::U64));
-			Members.AddRange(Batch.NameMember("StructRRs"), StructRRs.BuildAndReset(ERangeSizeType::U32));
+			Members.AddRange(Batch.NameMember("IntRs"), IntRs.BuildAndReset(Scratch, ERangeSizeType::S32));
+			Members.AddRange(Batch.NameMember("EmptyRs"), IntRs.BuildAndReset(Scratch, ERangeSizeType::S32));
+			Members.AddRange(Batch.NameMember("EnumRs"), EnumRs.BuildAndReset(Scratch, ERangeSizeType::U8));
+			Members.AddRange(Batch.NameMember("StructRs"), StructRs.BuildAndReset(Scratch, ERangeSizeType::U64));
+			Members.AddRange(Batch.NameMember("StructRRs"), StructRRs.BuildAndReset(Scratch, ERangeSizeType::U32));
 
 			Batch.AddObject(Object, MoveTemp(Members));
 		}, 
@@ -923,7 +925,7 @@ TEST_CASE_NAMED(FPlainPropsReadWriteTest, "System::Core::Serialization::PlainPro
 	
 	SECTION("UniRange")
 	{
-		TestSerialize([](FTestBatchBuilder& Batch)
+		TestSerialize([](FTestBatchBuilder& Batch, FScratchAllocator& Scratch)
 		{
 			FStructSchemaId Object = Batch.DeclareStruct("Test", "Object", {"Bools", "Structs", "BF", "BT" }, EMemberPresence::AllowSparse);
 			FStructSchemaId Struct = Batch.DeclareStruct("Test", "Struct", {"MaybeB", "Bs", "MaybeBs", "B"}, EMemberPresence::AllowSparse);
@@ -932,20 +934,20 @@ TEST_CASE_NAMED(FPlainPropsReadWriteTest, "System::Core::Serialization::PlainPro
 			const bool False = false;
 			FNestedRangeBuilder MaybeBs(MakeLeafRangeSchema<bool, bool>(), 1);
 			FStructRangeBuilder Structs(10);
-			Structs[5].AddRange(Batch.NameMember("MaybeB"), BuildLeafRange(&False, true));
-			Structs[6].AddRange(Batch.NameMember("MaybeB"), BuildLeafRange(&True, false));
-			Structs[7].AddRange(Batch.NameMember("MaybeB"), BuildLeafRange(&True, true));
-			Structs[7].AddRange(Batch.NameMember("Bs"),		BuildLeafRange(MakeArrayView({true, true, false, false, true, true, false, false, true, true})));
-			MaybeBs.Add(BuildLeafRange(&True, true));
-			Structs[7].AddRange(Batch.NameMember("MaybeBs"), MaybeBs.BuildAndReset(ERangeSizeType::Uni));
+			Structs[5].AddRange(Batch.NameMember("MaybeB"), BuildLeafRange(Scratch, &False, true));
+			Structs[6].AddRange(Batch.NameMember("MaybeB"), BuildLeafRange(Scratch, &True, false));
+			Structs[7].AddRange(Batch.NameMember("MaybeB"), BuildLeafRange(Scratch, &True, true));
+			Structs[7].AddRange(Batch.NameMember("Bs"),		BuildLeafRange(Scratch, MakeArrayView({true, true, false, false, true, true, false, false, true, true})));
+			MaybeBs.Add(BuildLeafRange(Scratch, &True, true));
+			Structs[7].AddRange(Batch.NameMember("MaybeBs"), MaybeBs.BuildAndReset(Scratch, ERangeSizeType::Uni));
 			Structs[7].Add(Batch.NameMember("B"), true);
-			MaybeBs.Add(BuildLeafRange(&True, false));
-			Structs[8].AddRange(Batch.NameMember("MaybeBs"), MaybeBs.BuildAndReset(ERangeSizeType::Uni));
+			MaybeBs.Add(BuildLeafRange(Scratch, &True, false));
+			Structs[8].AddRange(Batch.NameMember("MaybeBs"), MaybeBs.BuildAndReset(Scratch, ERangeSizeType::Uni));
 			Structs[9].Add(Batch.NameMember("B"), false);
 
 			FMemberBuilder Members;
-			Members.AddRange(Batch.NameMember("Bools"), BuildLeafRange(&True, true));
-			Members.AddRange(Batch.NameMember("Structs"), Structs.BuildAndReset(Batch.Get(Struct), Batch));
+			Members.AddRange(Batch.NameMember("Bools"), BuildLeafRange(Scratch, &True, true));
+			Members.AddRange(Batch.NameMember("Structs"), Structs.BuildAndReset(Scratch, Batch.Get(Struct), Batch));
 			Members.Add(Batch.NameMember("BF"), false);
 			Members.Add(Batch.NameMember("BT"), true);
 
@@ -981,7 +983,7 @@ TEST_CASE_NAMED(FPlainPropsReadWriteTest, "System::Core::Serialization::PlainPro
 	
 	SECTION("DynamicStruct")
 	{
-		TestSerialize([](FTestBatchBuilder& Batch)
+		TestSerialize([](FTestBatchBuilder& Batch, FScratchAllocator& Scratch)
 		{
 			FStructSchemaId Unused1 = Batch.DeclareStruct("Test", "Unused1", {"X"}, EMemberPresence::AllowSparse);
 			FStructSchemaId SA = Batch.DeclareStruct("Test", "SA", {"X"}, EMemberPresence::AllowSparse);
@@ -990,11 +992,11 @@ TEST_CASE_NAMED(FPlainPropsReadWriteTest, "System::Core::Serialization::PlainPro
 			FStructSchemaId Object = Batch.DeclareStruct("Test", "Object", {"Same", "Some", "None", "Diff"}, EMemberPresence::AllowSparse);
 			FStructSchemaId Unused3 = Batch.DeclareStruct("Test", "Unused3", {"X"}, EMemberPresence::AllowSparse);
 		
-			auto BuildStruct = [&Batch](FStructSchemaId Struct, auto X)
+			auto BuildStruct = [&](FStructSchemaId Struct, auto X)
 				{
 					FMemberBuilder Members;
 					Members.Add(Batch.NameMember("X"), X);
-					return Members.BuildAndReset(Batch.Get(Struct), Batch);
+					return Members.BuildAndReset(Scratch, Batch.Get(Struct), Batch);
 				};
 
 			FMemberBuilder O1;
@@ -1030,18 +1032,18 @@ TEST_CASE_NAMED(FPlainPropsReadWriteTest, "System::Core::Serialization::PlainPro
 
 	SECTION("DynamicStructRange")
 	{
-		TestSerialize([](FTestBatchBuilder& Batch)
+		TestSerialize([](FTestBatchBuilder& Batch, FScratchAllocator& Scratch)
 		{
 			FStructSchemaId SA = Batch.DeclareStruct("Test", "SA", {"X"}, EMemberPresence::AllowSparse);
 			FStructSchemaId Unused = Batch.DeclareStruct("Test", "Unused2", {"X"}, EMemberPresence::AllowSparse);
 			FStructSchemaId SB = Batch.DeclareStruct("Test", "SB", {"X"}, EMemberPresence::AllowSparse);
 			FStructSchemaId Object = Batch.DeclareStruct("Test", "Object", {"Same", "Some", "None", "Diff", "SameEmpty", "DiffEmpty", "DiffNested"}, EMemberPresence::AllowSparse);
 		
-			auto BuildStructRange = [&Batch](FStructSchemaId Struct, auto X)
+			auto BuildStructRange = [&](FStructSchemaId Struct, auto X)
 				{
 					FStructRangeBuilder Members(1);
 					Members[0].Add(Batch.NameMember("X"), X);
-					return Members.BuildAndReset(Batch.Get(Struct), Batch);
+					return Members.BuildAndReset(Scratch, Batch.Get(Struct), Batch);
 				};
 
 			
@@ -1053,16 +1055,16 @@ TEST_CASE_NAMED(FPlainPropsReadWriteTest, "System::Core::Serialization::PlainPro
 			O1.AddRange(Batch.NameMember("DiffEmpty"), BuildStructRange(SA, 14));
 			FNestedRangeBuilder NestedSA(MakeStructRangeSchema(ERangeSizeType::S32, SA), 1);
 			NestedSA.Add(BuildStructRange(SA, 100));
-			O1.AddRange(Batch.NameMember("DiffNested"), NestedSA.BuildAndReset(ERangeSizeType::S32));
+			O1.AddRange(Batch.NameMember("DiffNested"), NestedSA.BuildAndReset(Scratch, ERangeSizeType::S32));
 			
 			FMemberBuilder O2;
 			O2.AddRange(Batch.NameMember("Same"), BuildStructRange(SA, 20));
 			O2.AddRange(Batch.NameMember("Diff"), BuildStructRange(SB, 22.f));
-			O2.AddRange(Batch.NameMember("SameEmpty"), FStructRangeBuilder(0).BuildAndReset(Batch.Get(SA), Batch));
-			O2.AddRange(Batch.NameMember("DiffEmpty"), FStructRangeBuilder(0).BuildAndReset(Batch.Get(SB), Batch));
+			O2.AddRange(Batch.NameMember("SameEmpty"), FStructRangeBuilder(0).BuildAndReset(Scratch, Batch.Get(SA), Batch));
+			O2.AddRange(Batch.NameMember("DiffEmpty"), FStructRangeBuilder(0).BuildAndReset(Scratch, Batch.Get(SB), Batch));
 			FNestedRangeBuilder NestedSB(MakeStructRangeSchema(ERangeSizeType::S32, SB), 1);
 			NestedSB.Add(BuildStructRange(SB, 200.f));
-			O2.AddRange(Batch.NameMember("DiffNested"), NestedSB.BuildAndReset(ERangeSizeType::S32));
+			O2.AddRange(Batch.NameMember("DiffNested"), NestedSB.BuildAndReset(Scratch, ERangeSizeType::S32));
 			
 			Batch.AddObject(Object, MoveTemp(O1));
 			Batch.AddObject(Object, MoveTemp(O2));
@@ -1092,7 +1094,7 @@ TEST_CASE_NAMED(FPlainPropsReadWriteTest, "System::Core::Serialization::PlainPro
 
 	SECTION("Inheritance")
 	{
-		TestSerialize([](FTestBatchBuilder& Batch)
+		TestSerialize([](FTestBatchBuilder& Batch, FScratchAllocator& Scratch)
 		{
 			FStructSchemaId Unused = Batch.DeclareStruct("Test", "X", {"X"}, EMemberPresence::AllowSparse);
 			FStructSchemaId Low = Batch.DeclareStruct("Test", "Low", {"LInt"}, EMemberPresence::AllowSparse);
@@ -1101,18 +1103,18 @@ TEST_CASE_NAMED(FPlainPropsReadWriteTest, "System::Core::Serialization::PlainPro
 			
 			FMemberBuilder Members;
 			Members.Add(Batch.NameMember("LInt"), 123);
-			Members.BuildSuperStruct(Batch.Get(Low), Batch);
+			Members.BuildSuperStruct(Scratch, Batch.Get(Low), Batch);
 			Members.Add(Batch.NameMember("MInt"), 456);
 			FMemberBuilder Nested;
 			Nested.Add(Batch.NameMember("LInt"), 1000);
-			Members.AddStruct(Batch.NameMember("MLow"), Low, Nested.BuildAndReset(Batch.Get(Low), Batch));
-			Members.BuildSuperStruct(Batch.Get(Mid), Batch);
+			Members.AddStruct(Batch.NameMember("MLow"), Low, Nested.BuildAndReset(Scratch, Batch.Get(Low), Batch));
+			Members.BuildSuperStruct(Scratch, Batch.Get(Mid), Batch);
 			Members.Add(Batch.NameMember("TInt"), 789);
 			Nested.Add(Batch.NameMember("LInt"), 2000);
-			Members.AddStruct(Batch.NameMember("TLow"), Low, Nested.BuildAndReset(Batch.Get(Low), Batch));
+			Members.AddStruct(Batch.NameMember("TLow"), Low, Nested.BuildAndReset(Scratch, Batch.Get(Low), Batch));
 			FStructRangeBuilder NestedRange(1);
 			NestedRange[0].Add(Batch.NameMember("MInt"), 3000);
-			Members.AddRange(Batch.NameMember("TMids"), NestedRange.BuildAndReset(Batch.Get(Mid), Batch));
+			Members.AddRange(Batch.NameMember("TMids"), NestedRange.BuildAndReset(Scratch, Batch.Get(Mid), Batch));
 
 			Batch.AddObject(Top, MoveTemp(Members));
 		}, 
@@ -1152,7 +1154,7 @@ TEST_CASE_NAMED(FPlainPropsReadWriteTest, "System::Core::Serialization::PlainPro
 
 	SECTION("SparseInheritance")
 	{
-		TestSerialize([](FTestBatchBuilder& Batch)
+		TestSerialize([](FTestBatchBuilder& Batch, FScratchAllocator& Scratch)
 		{																										// Usage by A B C 
 			FStructSchemaId B0 = Batch.DeclareStruct("Test", "B0", {"0"}, EMemberPresence::AllowSparse);				 // - - -
 			FStructSchemaId B1 = Batch.DeclareStruct("Test", "B1", {"1"}, EMemberPresence::AllowSparse, ToOptional(B0)); // 1 - 0
@@ -1165,26 +1167,26 @@ TEST_CASE_NAMED(FPlainPropsReadWriteTest, "System::Core::Serialization::PlainPro
 			
 			FMemberBuilder A;
 			A.Add(Batch.NameMember("1"), 1);
-			A.BuildSuperStruct(Batch.Get(B1), Batch);
+			A.BuildSuperStruct(Scratch, Batch.Get(B1), Batch);
 			A.Add(Batch.NameMember("4"), 4);
-			A.BuildSuperStruct(Batch.Get(B4), Batch);
+			A.BuildSuperStruct(Scratch, Batch.Get(B4), Batch);
 			A.Add(Batch.NameMember("5"), 5);
-			A.BuildSuperStruct(Batch.Get(B5), Batch);
+			A.BuildSuperStruct(Scratch, Batch.Get(B5), Batch);
 
 			FMemberBuilder B;
 			B.Add(Batch.NameMember("2"), 20);
-			B.BuildSuperStruct(Batch.Get(B2), Batch);
+			B.BuildSuperStruct(Scratch, Batch.Get(B2), Batch);
 			B.Add(Batch.NameMember("4"), 40);
-			B.BuildSuperStruct(Batch.Get(B4), Batch);
+			B.BuildSuperStruct(Scratch, Batch.Get(B4), Batch);
 
 			FMemberBuilder C;
-			C.BuildSuperStruct(Batch.Get(B1), Batch); // Empty -> noop
+			C.BuildSuperStruct(Scratch, Batch.Get(B1), Batch); // Empty -> noop
 			C.Add(Batch.NameMember("2"), 200);
-			C.BuildSuperStruct(Batch.Get(B2), Batch);
-			C.BuildSuperStruct(Batch.Get(B3), Batch); // Empty -> noop
+			C.BuildSuperStruct(Scratch, Batch.Get(B2), Batch);
+			C.BuildSuperStruct(Scratch, Batch.Get(B3), Batch); // Empty -> noop
 			C.Add(Batch.NameMember("4"), 400);
-			C.BuildSuperStruct(Batch.Get(B4), Batch);
-			C.BuildSuperStruct(Batch.Get(B5), Batch); // Empty -> noop
+			C.BuildSuperStruct(Scratch, Batch.Get(B4), Batch);
+			C.BuildSuperStruct(Scratch, Batch.Get(B5), Batch); // Empty -> noop
 			
 			Batch.AddObject(B6, MoveTemp(A));
 			Batch.AddObject(B5, MoveTemp(B));
@@ -1229,7 +1231,7 @@ TEST_CASE_NAMED(FPlainPropsReadWriteTest, "System::Core::Serialization::PlainPro
 
 	SECTION("SparseIndex")
 	{
-		TestSerialize([](FTestBatchBuilder& Batch)
+		TestSerialize([](FTestBatchBuilder& Batch, FScratchAllocator& Scratch)
 		{
 			FScopeId Unused = Batch.MakeScope("Unused");
 			FScopeId NestedUnused1 = Batch.NestScope(Unused, "NestedUnused1");
diff --git a/Engine/Plugins/Experimental/PlainProps/Source/Tests/PlainPropsUeCoreTests.cpp b/Engine/Plugins/Experimental/PlainProps/Source/Tests/PlainPropsUeCoreTests.cpp
--- a/Engine/Plugins/Experimental/PlainProps/Source/Tests/PlainPropsUeCoreTests.cpp
+++ b/Engine/Plugins/Experimental/PlainProps/Source/Tests/PlainPropsUeCoreTests.cpp
@@ -186,6 +186,7 @@ private:
 	TArray<IdBuiltStructPair>	SavedObjects;
 	FNameBinding				SavedNames;
 	FCustomBindings				Customs;
+	mutable FScratchAllocator	Scratch;
 };
 
 FBatchSaver::FBatchSaver()
@@ -198,14 +199,14 @@ template<class T>
 void FBatchSaver::Save(T&& Object) 
 {
 	FStructSchemaId Id = IndexNativeStruct<T, FIds>();
-	SavedObjects.Emplace(Id, SaveStruct(&Object, Id, {GTypes, GSchemas, Customs}));
+	SavedObjects.Emplace(Id, SaveStruct(&Object, Id, {GTypes, GSchemas, Customs, Scratch}));
 }
 
 template<class T>
 bool FBatchSaver::SaveDelta(const T& Object, const T& Default) 
 {
 	FStructSchemaId Id = IndexNativeStruct<T, FIds>();
-	if (FBuiltStructPtr Delta = SaveStructDelta(&Object, &Default, Id, {GTypes, GSchemas, Customs}))
+	if (FBuiltStructPtr Delta = SaveStructDelta(&Object, &Default, Id, {GTypes, GSchemas, Customs, Scratch}))
 	{
 		SavedObjects.Emplace(Id, MoveTemp(Delta));
 		return true;
@@ -230,7 +231,7 @@ TConstArrayView<T> GrabNumAndArray(/* in-out */ FByteReader& It)
 TArray64<uint8> FBatchSaver::Write() const
 {
 	// Build partial schemas
-	FSchemasBuilder SchemaBuilders(GTypes.GetStructs(), GTypes.GetEnums(), /* debug */ GNames);
+	FSchemasBuilder SchemaBuilders(GTypes, Scratch);
 	for (const IdBuiltStructPair& Object : SavedObjects)
 	{
 		SchemaBuilders.NoteStructAndMembers(Object.Key, *Object.Value);
