[PlainProps]
* Disabled static analysis since code is not ready for production and Johan doesn't want to have to care about static analysis

diff --git a/Engine/Plugins/Experimental/PlainProps/Source/PlainProps.Build.cs b/Engine/Plugins/Experimental/PlainProps/Source/PlainProps.Build.cs
--- a/Engine/Plugins/Experimental/PlainProps/Source/PlainProps.Build.cs
+++ b/Engine/Plugins/Experimental/PlainProps/Source/PlainProps.Build.cs
@@ -6,6 +6,8 @@ public class PlainProps : ModuleRules
 {
 	public PlainProps(ReadOnlyTargetRules Target) : base(Target)
 	{
+		bDisableStaticAnalysis = true;
+
 		CppStandard = CppStandardVersion.Cpp20;
 
 		PrivateDependencyModuleNames.AddRange(
