PlainProps load optimizations, ~100ms -> 72ms for point cloud LoadStruct()

#rnx

diff --git a/Engine/Plugins/Experimental/PlainProps/Source/Private/PlainPropsRead.cpp b/Engine/Plugins/Experimental/PlainProps/Source/Private/PlainPropsRead.cpp
--- a/Engine/Plugins/Experimental/PlainProps/Source/Private/PlainPropsRead.cpp
+++ b/Engine/Plugins/Experimental/PlainProps/Source/Private/PlainPropsRead.cpp
@@ -183,19 +183,6 @@ FNestedRangeView FRangeView::AsRanges() const
 
 //////////////////////////////////////////////////////////////////////////
 
-bool FBitCacheReader::GrabNext(FByteReader& Bytes)
-{
-	BitIt <<= 1; // Shift up til overflow
-
-	if (BitIt == 0)
-	{
-		Bits = Bytes.GrabByte();
-		BitIt = 1;
-	}
-
-	return !!(Bits & BitIt);
-}
-
 const FStructSchema& FStructSchemaHandle::ResolveSuper() const
 {
 	return ResolveStructSchema(Batch, Resolve().GetSuperSchema().Get());
@@ -249,10 +236,10 @@ static bool SkipDeclaredSuperSchema(ESuper Inheritance)
 FMemberReader::FMemberReader(const FStructSchema& Schema, FByteReader Values, FReadBatchId InBatch)
 : Footer(Schema.Footer)
 , Batch(InBatch)
-, NumMembers(Schema.NumMembers)
-, NumRangeTypes(Schema.NumRangeTypes)
 , IsSparse(!Schema.IsDense)
 , HasSuper(UsesSuper(Schema.Inheritance))
+, NumMembers(Schema.NumMembers)
+, NumRangeTypes(Schema.NumRangeTypes)
 , InnerSchemaIdx(SkipDeclaredSuperSchema(Schema.Inheritance))
 , ValueIt(Values)
 #if DO_CHECK
@@ -261,20 +248,25 @@ FMemberReader::FMemberReader(const FStructSchema& Schema, FByteReader Values, FR
 {
 	check(InnerSchemaIdx <= NumInnerSchemas);
 	checkf(NumRangeTypes != 0xFFFFu, TEXT("GrabRangeTypes() doesn't check for wrap-around"));
-	SkipMissingSparseMembers();
+
+	if (IsSparse)
+	{
+		SkipMissingSparseMembers();
+	}
 }
 
 FOptionalMemberId FMemberReader::PeekName() const
 {
 	int32 MemberNameIdx = MemberIdx - HasSuper;
-	if (MemberNameIdx == -1)
-	{
-		UE_DEBUG_BREAK();
-	}
-
 	return MemberNameIdx >= 0 ? ToOptional(GetMemberNames()[MemberNameIdx]) : NoId;
 }
 
+FOptionalMemberId FMemberReader::PeekNameUnchecked() const
+{
+	int32 MemberNameIdx = MemberIdx - HasSuper;
+	return GetMemberNames()[MemberNameIdx];
+}
+
 EMemberKind FMemberReader::PeekKind() const
 {
 	return PeekType().GetKind();
@@ -289,24 +281,25 @@ FMemberType	FMemberReader::PeekType() const
 void FMemberReader::AdvanceToNextMember()
 {
 	++MemberIdx;
-	SkipMissingSparseMembers();
+	if (IsSparse)
+	{
+		SkipMissingSparseMembers();
+	}
 }
 
 void FMemberReader::SkipMissingSparseMembers()
 {
-	if (IsSparse)
+	// Change code in LoadMembers() too
+	while (MemberIdx < NumMembers && GrabBit())
 	{
-		// Change code in LoadMembers() too
-		while (MemberIdx < NumMembers && GrabBit())
-		{
-			FMemberType InnermostType = PeekType().IsRange() ? /* advances RangeTypeIdx */ GrabRangeTypes().Last() : PeekType();
-			SkipSchema(InnermostType);
-			++MemberIdx;
-		}
-	}	
+		FMemberType Type = GetMemberTypes()[MemberIdx];
+		FMemberType InnermostType = Type.IsRange() ? /* advances RangeTypeIdx */ GrabRangeTypes().Last() : Type;
+		SkipSchema(InnermostType);
+		++MemberIdx;
+	}
 }
 
-void FMemberReader::SkipSchema(FMemberType InnermostType)
+inline void FMemberReader::SkipSchema(FMemberType InnermostType)
 {
 	if (InnermostType.IsStruct())
 	{
@@ -353,23 +346,21 @@ FOptionalSchemaId FMemberReader::GrabRangeSchema(FMemberType InnermostType)
 
 FLeafView FMemberReader::GrabLeaf()
 {
-	check(HasMore());
-	FUnpackedLeafType Leaf = PeekType().AsLeaf();
-	FEnumSchemaId Enum = Leaf.Type == ELeafType::Enum ? GrabEnumSchema() : FEnumSchemaId{};
+	FLeafView Out = { PeekType().AsLeaf(), Batch };
+	Out.Enum = Out.Leaf.Type == ELeafType::Enum ? GrabEnumSchema() : FEnumSchemaId{};
 	
-	FMemberValue Value;
-	if (Leaf.Type == ELeafType::Bool)
+	if (Out.Leaf.Type == ELeafType::Bool)
 	{
-		Value.bValue = GrabBit();
+		Out.Value.bValue = GrabBit();
 	}
 	else
 	{
-		Value.Ptr = ValueIt.GrabBytes(SizeOf(Leaf.Width));
+		Out.Value.Ptr = ValueIt.GrabBytes(SizeOf(Out.Leaf.Width));
 	}
 
 	AdvanceToNextMember();
 
-	return { Leaf.Type, Leaf.Width, Batch, Enum, Value };
+	return Out;
 }
 
 FStructView FMemberReader::GrabStruct()
@@ -414,6 +405,44 @@ FRangeView FMemberReader::GrabRange()
 //	return {Type, Id, BatchId, Value};
 //}
 
+
+void FMemberReader::GrabLeaves(void* Out, uint32 Num, SIZE_T Size)
+{
+	checkSlow(Num);
+	check(MemberIdx + Num <= NumMembers);
+	const FMemberType* Types = GetMemberTypes() + MemberIdx;
+	FLeafType Leaf = Types[0].AsLeaf();
+	checkSlow(Leaf.Type != ELeafType::Enum);
+	checkSlow(SizeOf(Leaf.Width) == Size);
+	for (FMemberType Type : MakeArrayView(Types + 1, Num - 1))
+	{
+		check(Type == Types[0]);
+	}
+
+	
+	if (IsSparse)
+	{
+		uint8* OutIt = static_cast<uint8*>(Out);
+		for (uint8* OutLast = OutIt + Num * Size - Size; OutIt != OutLast; OutIt += Size)
+		{
+			FMemory::Memcpy(OutIt, ValueIt.GrabBytes(Size), Size);
+			bool bSkip = GrabBit();
+			check(!bSkip);
+		}
+
+		FMemory::Memcpy(OutIt, ValueIt.GrabBytes(Size), Size);
+		
+		MemberIdx += Num;
+		SkipMissingSparseMembers();
+	}
+	else
+	{
+		const SIZE_T NumBytes = Num * Size;
+		FMemory::Memcpy(Out, ValueIt.GrabBytes(NumBytes), NumBytes);
+		MemberIdx += Num;
+	}
+}
+
 //////////////////////////////////////////////////////////////////////////
 
 static TOptional<FStructView> TryGrabSuper(FMemberReader& Members)
diff --git a/Engine/Plugins/Experimental/PlainProps/Source/Private/PlainPropsUeCoreBindings.cpp b/Engine/Plugins/Experimental/PlainProps/Source/Private/PlainPropsUeCoreBindings.cpp
--- a/Engine/Plugins/Experimental/PlainProps/Source/Private/PlainPropsUeCoreBindings.cpp
+++ b/Engine/Plugins/Experimental/PlainProps/Source/Private/PlainPropsUeCoreBindings.cpp
@@ -7,7 +7,7 @@
 namespace PlainProps::UE
 {
 
-void FTransformBinding::Save(FMemberBuilder& Dst, const FTransform& Src, const FTransform* Default, const FTransformIds& Ids) const
+void FTransformBinding::Save(FMemberBuilder& Dst, const FTransform& Src, const FTransform* Default) const
 {
 	static_assert(std::is_same_v<decltype(FTransform().GetTranslation().X), double>);
 
@@ -19,43 +19,45 @@ void FTransformBinding::Save(FMemberBuilder& Dst, const FTransform& Src, const F
 	{
 		if (T != Default->GetTranslation())
 		{
-			Dst.Add(Ids.Translate[0], T.X);
-			Dst.Add(Ids.Translate[1], T.Y);
-			Dst.Add(Ids.Translate[2], T.Z);
+			Dst.Add(MemberIds.Translate[0], T.X);
+			Dst.Add(MemberIds.Translate[1], T.Y);
+			Dst.Add(MemberIds.Translate[2], T.Z);
 		}
 
 		if (R != Default->GetRotation())
 		{
-			Dst.Add(Ids.Rotate[0], R.X);
-			Dst.Add(Ids.Rotate[1], R.Y);
-			Dst.Add(Ids.Rotate[2], R.Z);
-			Dst.Add(Ids.Rotate[3], R.W);
+			Dst.Add(MemberIds.Rotate[0], R.X);
+			Dst.Add(MemberIds.Rotate[1], R.Y);
+			Dst.Add(MemberIds.Rotate[2], R.Z);
+			Dst.Add(MemberIds.Rotate[3], R.W);
 		}
 
 		if (S != Default->GetScale3D())
 		{
-			Dst.Add(Ids.Scale[0], S.X);
-			Dst.Add(Ids.Scale[1], S.Y);
-			Dst.Add(Ids.Scale[2], S.Z);
+			Dst.Add(MemberIds.Scale[0], S.X);
+			Dst.Add(MemberIds.Scale[1], S.Y);
+			Dst.Add(MemberIds.Scale[2], S.Z);
 		}
 	}
 	else
 	{
-		Dst.Add(Ids.Translate[0],	T.X);
-		Dst.Add(Ids.Translate[1],	T.Y);
-		Dst.Add(Ids.Translate[2],	T.Z);
-		Dst.Add(Ids.Rotate[0],		R.X);
-		Dst.Add(Ids.Rotate[1],		R.Y);
-		Dst.Add(Ids.Rotate[2],		R.Z);
-		Dst.Add(Ids.Rotate[3],		R.W);
-		Dst.Add(Ids.Scale[0],		S.X);
-		Dst.Add(Ids.Scale[1],		S.Y);
-		Dst.Add(Ids.Scale[2],		S.Z);
+		Dst.Add(MemberIds.Translate[0],	T.X);
+		Dst.Add(MemberIds.Translate[1],	T.Y);
+		Dst.Add(MemberIds.Translate[2],	T.Z);
+		Dst.Add(MemberIds.Rotate[0],		R.X);
+		Dst.Add(MemberIds.Rotate[1],		R.Y);
+		Dst.Add(MemberIds.Rotate[2],		R.Z);
+		Dst.Add(MemberIds.Rotate[3],		R.W);
+		Dst.Add(MemberIds.Scale[0],		S.X);
+		Dst.Add(MemberIds.Scale[1],		S.Y);
+		Dst.Add(MemberIds.Scale[2],		S.Z);
 	}
 }
 
-void FTransformBinding::Load(FTransform& Dst, FStructView Src, ECustomLoadMethod Method, const FLoadBatch& Batch, const FTransformIds& Ids) const
+void FTransformBinding::Load(FTransform& Dst, FStructView Src, ECustomLoadMethod Method, const FLoadBatch& Batch) const
 {
+	static_assert(std::is_same_v<decltype(FTransform().GetTranslation().X), double>);
+
 	FMemberReader Members(Src);
 
 	if (Method == ECustomLoadMethod::Construct)
@@ -68,45 +70,35 @@ void FTransformBinding::Load(FTransform& Dst, FStructView Src, ECustomLoadMethod
 		return;
 	}
 
-	FMemberId Name = Members.PeekName().Get();
-	double X = Members.GrabLeaf().AsDouble();
-	double Y = Members.GrabLeaf().AsDouble();
-	double Z = Members.GrabLeaf().AsDouble();
-
-	if (Name == Ids.Translate[0])
+	if (Members.PeekNameUnchecked() == MemberIds.Translate[0])
 	{
-		Dst.SetTranslation({X, Y, Z});
-			
+		FVector Translation;
+		Members.GrabLeaves(&Translation.X, 3);
+		Dst.SetTranslation(Translation);
+
 		if (!Members.HasMore())
 		{
 			return;
 		}
-			
-		Name = Members.PeekName().Get();
-		X = Members.GrabLeaf().AsDouble();
-		Y = Members.GrabLeaf().AsDouble();
-		Z = Members.GrabLeaf().AsDouble();
 	}
 
-	if (Name == Ids.Rotate[0])
+	if (Members.PeekNameUnchecked() == MemberIds.Rotate[0])
 	{
-		double W = Members.GrabLeaf().AsDouble();
-		Dst.SetRotation({X, Y, Z, W});
+		FQuat Rotation;
+		Members.GrabLeaves(&Rotation.X, 4);
+		Dst.SetRotation(Rotation);
 
 		if (!Members.HasMore())
 		{
 			return;
 		}
-
-		Name = Members.PeekName().Get();
-		X = Members.GrabLeaf().AsDouble();
-		Y = Members.GrabLeaf().AsDouble();
-		Z = Members.GrabLeaf().AsDouble();
 	}
 
-	check(Name == Ids.Scale[0]);
-	Dst.SetScale3D({X, Y, Z});
-	check(!Members.HasMore());
+	checkSlow(Members.PeekNameUnchecked() == MemberIds.Scale[0]);
+	FVector Scale;
+	Members.GrabLeaves(&Scale.X, 3);
+	Dst.SetScale3D(Scale);
+	checkSlow(!Members.HasMore());
 }
 
 bool FTransformBinding::DiffStruct(const void* StructA, const void* StructB) const
diff --git a/Engine/Plugins/Experimental/PlainProps/Source/Public/PlainPropsRead.h b/Engine/Plugins/Experimental/PlainProps/Source/Public/PlainPropsRead.h
--- a/Engine/Plugins/Experimental/PlainProps/Source/Public/PlainPropsRead.h
+++ b/Engine/Plugins/Experimental/PlainProps/Source/Public/PlainPropsRead.h
@@ -63,23 +63,23 @@ public:
 		return Out;
 	}
 
-	[[nodiscard]] FMemoryView GrabSlice(uint64 NumBytes)
+	[[nodiscard]] inline FMemoryView GrabSlice(uint64 NumBytes)
 	{
 		return FMemoryView(GrabBytes(NumBytes), NumBytes);
 	}
 
-	[[nodiscard]] FMemoryView GrabSkippableSlice()
+	[[nodiscard]] inline FMemoryView GrabSkippableSlice()
 	{
 		return GrabSlice(GrabVarIntU());
 	}
 
-	[[nodiscard]] uint8 GrabByte()
+	[[nodiscard]] inline uint8 GrabByte()
 	{
 		return *GrabBytes(1);
 	}
 
 	template<typename T>
-	[[nodiscard]] T Grab()
+	[[nodiscard]] inline T Grab()
 	{
 		return FPlatformMemory::ReadUnaligned<T>(GrabBytes(sizeof(T)));
 	}
@@ -118,65 +118,95 @@ class FBitCacheReader
 	uint8 Bits = 0;
 	uint8 BitIt = 0;
 public:
-	[[nodiscard]] bool GrabNext(FByteReader& Bytes);
+	[[nodiscard]] FORCEINLINE bool GrabNext(FByteReader& Bytes)
+	{
+		BitIt <<= 1; // Shift up til overflow
+
+		if (BitIt == 0)
+		{
+			Bits = Bytes.GrabByte();
+			BitIt = 1;
+		}
+
+		return !!(Bits & BitIt);
+	}
+
+	FORCENOINLINE void Skip(uint32 Num, FByteReader& Bytes)
+	{
+		uint32 NumCached = 1 + FMath::CountLeadingZeros8(BitIt);
+
+		if (NumCached > Num)
+		{
+			BitIt <<= Num;
+		}
+		else
+		{
+			uint32 NumUncached = Num - NumCached;
+
+			// Grab new bytes, keep the last byte and bit within it
+			uint32 NumBytes = Align(NumUncached + 1, 8) / 8;
+			Bits = Bytes.GrabBytes(NumBytes)[NumBytes - 1];
+			BitIt = 1 << (NumUncached % 8);
+		}
+	}
 };
 
 //////////////////////////////////////////////////////////////////////////
 
 struct FStructSchemaHandle
 {
-	FStructSchemaId		Id;
-	FReadBatchId		Batch;
+	FStructSchemaId			Id;
+	FReadBatchId			Batch;
 
-	const FStructSchema& Resolve() const { return ResolveStructSchema(Batch, Id); }
-	const FStructSchema& ResolveSuper() const;
+	const FStructSchema&	Resolve() const { return ResolveStructSchema(Batch, Id); }
+	const FStructSchema&	ResolveSuper() const;
 };
 
 struct FStructView
 {
-	FStructSchemaHandle	Schema;
-	FByteReader			Values;
+	FStructSchemaHandle		Schema;
+	FByteReader				Values;
 };
 
 //////////////////////////////////////////////////////////////////////////
 
 union FMemberValue
 {
-	const uint8*		Ptr;	// From byte stream
-	bool				bValue; // From bit cache
+	const uint8*			Ptr;	// From byte stream
+	bool					bValue; // From bit cache
 };
 
 struct FLeafView
 {
-	ELeafType		Type;
-    ELeafWidth		Width;
-	FReadBatchId	Batch;
-	FEnumSchemaId	Enum;
-	FMemberValue	Value;
-
-	bool			AsBool() const		{ check(Type == ELeafType::Bool);	return Value.bValue; }
-	int8			AsS8() const		{ check(Type == ELeafType::IntS);	return As<int8>(); }
-	uint8			AsU8() const		{ check(Type == ELeafType::IntU);	return As<uint8>(); }
-	int16			AsS16() const		{ check(Type == ELeafType::IntS);	return As<int16>(); }
-	uint16			AsU16() const		{ check(Type == ELeafType::IntU);	return As<uint16>(); }
-	int32			AsS32() const		{ check(Type == ELeafType::IntS);	return As<int32>(); }
-	uint32			AsU32() const		{ check(Type == ELeafType::IntU);	return As<uint32>(); }
-	int64			AsS64() const		{ check(Type == ELeafType::IntS);	return As<int64>(); }
-	uint64			AsU64() const		{ check(Type == ELeafType::IntU);	return As<uint64>(); }
-	double			AsDouble() const	{ check(Type == ELeafType::Float);	return As<double>(); }
-	float			AsFloat() const		{ check(Type == ELeafType::Float);	return As<float>(); }
-	char8_t			AsChar8() const		{ check(Type == ELeafType::Unicode);return As<char8_t>(); }
-	char16_t		AsChar16() const	{ check(Type == ELeafType::Unicode);return As<char16_t>(); }
-	char32_t		AsChar32() const	{ check(Type == ELeafType::Unicode);return As<char32_t>(); }
-	uint8			AsEnum8() const		{ check(Type == ELeafType::Enum);	return As<uint8>(); }
-	uint16			AsEnum16() const	{ check(Type == ELeafType::Enum);	return As<uint16>(); }
-	uint32			AsEnum32() const	{ check(Type == ELeafType::Enum);	return As<uint32>(); }
-	uint64			AsEnum64() const	{ check(Type == ELeafType::Enum);	return As<uint64>(); }
+	FUnpackedLeafType		Leaf;
+	FReadBatchId			Batch;
+	FEnumSchemaId			Enum;
+	FMemberValue			Value;
+
+	FORCEINLINE bool		AsBool() const		{ return Value.bValue; }
+	FORCEINLINE int8		AsS8() const		{ return As<int8>(); }
+	FORCEINLINE uint8		AsU8() const		{ return As<uint8>(); }
+	FORCEINLINE int16		AsS16() const		{ return As<int16>(); }
+	FORCEINLINE uint16		AsU16() const		{ return As<uint16>(); }
+	FORCEINLINE int32		AsS32() const		{ return As<int32>(); }
+	FORCEINLINE uint32		AsU32() const		{ return As<uint32>(); }
+	FORCEINLINE int64		AsS64() const		{ return As<int64>(); }
+	FORCEINLINE uint64		AsU64() const		{ return As<uint64>(); }
+	FORCEINLINE double		AsDouble() const	{ return As<double>(); }
+	FORCEINLINE float		AsFloat() const		{ return As<float>(); }
+	FORCEINLINE char8_t		AsChar8() const		{ return As<char8_t>(); }
+	FORCEINLINE char16_t	AsChar16() const	{ return As<char16_t>(); }
+	FORCEINLINE char32_t	AsChar32() const	{ return As<char32_t>(); }
+	FORCEINLINE uint8		AsEnum8() const		{ return As<uint8,	FUnpackedLeafType{ELeafType::Enum, ELeafWidth::B8}>(); }
+	FORCEINLINE uint16		AsEnum16() const	{ return As<uint16,	FUnpackedLeafType{ELeafType::Enum, ELeafWidth::B16}>(); }
+	FORCEINLINE uint32		AsEnum32() const	{ return As<uint32,	FUnpackedLeafType{ELeafType::Enum, ELeafWidth::B32}>(); }
+	FORCEINLINE uint64		AsEnum64() const	{ return As<uint64,	FUnpackedLeafType{ELeafType::Enum, ELeafWidth::B64}>(); }
 
 private:
-	template<typename T> T As() const
+	template<typename T, FUnpackedLeafType ExpectedType = ReflectLeaf<T>>
+	FORCEINLINE T As() const
 	{
-		check(SizeOf(Width) == sizeof(T));
+		check(ExpectedType == Leaf);
 		return *reinterpret_cast<const T*>(Value.Ptr);
 	}
 };
@@ -533,6 +563,7 @@ public:
 	bool					HasMore() const			{ return MemberIdx < NumMembers; }
 	
 	FOptionalMemberId		PeekName() const;		// @pre HasMore()
+	FOptionalMemberId		PeekNameUnchecked() const; // @pre HasMore()
 	EMemberKind				PeekKind() const;		// @pre HasMore()
 	FMemberType				PeekType() const;		// @pre HasMore()
 
@@ -541,22 +572,26 @@ public:
 	FStructView				GrabStruct();			// @pre PeekKind() == EMemberKind::Struct
 	//FAnyMemberView		GrabAny();				// @pre HasMore()
 
+	// @pre Has N more contiguous members of the expected leaf type
+	template<typename T>
+	void					GrabLeaves(T* Out, uint32 N);
+
 protected: // for unit tests
 	const FMemberType*		Footer;
 	const FReadBatchId		Batch;					// Needed to resolve schemas
-	const uint16			NumMembers;
-	const uint16			NumRangeTypes;			// Number of ranges and nested ranges
-	const uint16			IsSparse : 1;
-	const uint16			HasSuper : 1;
-
-	uint16					MemberIdx = 0;
-	uint16					RangeTypeIdx = 0;		// Types of [nested] ranges
-	uint16					InnerSchemaIdx = 0;		// Types of static structs and enums
+	const bool				IsSparse : 1;
+	const bool				HasSuper : 1;
+	const uint32			NumMembers;
+	const uint32			NumRangeTypes;			// Number of ranges and nested ranges
+
+	uint32					MemberIdx = 0;
+	uint32					RangeTypeIdx = 0;		// Types of [nested] ranges
+	uint32					InnerSchemaIdx = 0;		// Types of static structs and enums
 	FBitCacheReader			Bits;
 	FByteReader				ValueIt;
 
 #if DO_CHECK
-	const uint16			NumInnerSchemas;		// Number of static structs and enums
+	const uint32			NumInnerSchemas;		// Number of static structs and enums
 #endif
 
 	FMemberReader(const FStructSchema& Schema, FByteReader Values, FReadBatchId InBatch);
@@ -567,6 +602,7 @@ protected: // for unit tests
 	const FMemberId*		GetMemberNames() const;
 	
 	void					AdvanceToNextMember();
+	void					AdvanceToLaterMember(uint32 Num);
 	void					SkipMissingSparseMembers();
 	void					SkipSchema(FMemberType InnermostType);
 
@@ -576,11 +612,36 @@ protected: // for unit tests
 	FSchemaId				GrabInnerSchema();
 	FStructSchemaId			GrabStructSchema(FStructType Type);
 	FOptionalSchemaId		GrabRangeSchema(FMemberType InnermostType);
-	FEnumSchemaId			GrabEnumSchema()			{ return static_cast<FEnumSchemaId&&>(GrabInnerSchema()); }
-	bool					GrabBit()					{ return Bits.GrabNext(/* in-out */ ValueIt); }
-	uint64					GrabSkipLength()			{ return ValueIt.GrabVarIntU(); }
+	inline FEnumSchemaId	GrabEnumSchema()		{ return static_cast<FEnumSchemaId&&>(GrabInnerSchema()); }
+	inline bool				GrabBit()				{ return Bits.GrabNext(/* in-out */ ValueIt); }
+	inline uint64			GrabSkipLength()		{ return ValueIt.GrabVarIntU(); }
+
+	void					GrabBools(void* Out, uint32 Num);
+	void					GrabEnums(void* Out, uint32 Num, SIZE_T NumBytes);
+	void					GrabLeaves(void* Out, uint32 Num, SIZE_T NumBytes);
 };
 
+template<typename T>
+void FMemberReader::GrabLeaves(T* Out, uint32 N)
+{
+	if (N)
+	{
+		constexpr FUnpackedLeafType Leaf = ReflectLeaf<T>;
+		if constexpr (Leaf.Type == ELeafType::Bool)
+		{
+			GrabBools(Out, N);
+		}
+		else if constexpr (Leaf.Type == ELeafType::Enum)
+		{
+			GrabEnums(Out, N, sizeof(T));
+		}
+		else
+		{
+			GrabLeaves(Out, N, sizeof(T));
+		}
+	}
+}
+
 //////////////////////////////////////////////////////////////////////////
 
 // Hides the inheritance chain and iterates over super members first
diff --git a/Engine/Plugins/Experimental/PlainProps/Source/Public/PlainPropsTypes.h b/Engine/Plugins/Experimental/PlainProps/Source/Public/PlainPropsTypes.h
--- a/Engine/Plugins/Experimental/PlainProps/Source/Public/PlainPropsTypes.h
+++ b/Engine/Plugins/Experimental/PlainProps/Source/Public/PlainPropsTypes.h
@@ -315,11 +315,18 @@ struct FUnpackedLeafType
 	ELeafType Type;
 	ELeafWidth Width;
 
-	constexpr FUnpackedLeafType(ELeafType InType, ELeafWidth InWidth) : Type(InType), Width(InWidth) {}
-	constexpr FUnpackedLeafType(FLeafType In) : Type(In.Type), Width(In.Width) {}
+	inline constexpr FUnpackedLeafType(ELeafType InType, ELeafWidth InWidth) : Type(InType), Width(InWidth) {}
+	inline constexpr FUnpackedLeafType(FLeafType In) : Type(In.Type), Width(In.Width) {}
 
-	constexpr bool operator==(FUnpackedLeafType O) const { return Type == O.Type && Width == O.Width; }
-	FMemberType Pack() const { return FMemberType(Type, Width); }
+	inline bool operator==(FUnpackedLeafType O) const { return AsInt() == O.AsInt(); }
+	inline constexpr FMemberType Pack() const { return FMemberType(Type, Width); }
+
+	inline uint16 AsInt() const
+	{
+		uint16 Out = 0;
+		FMemory::Memcpy(&Out, this, sizeof(uint16));
+		return Out;
+	}
 };
 
 template<typename T>
diff --git a/Engine/Plugins/Experimental/PlainProps/Source/Public/PlainPropsUeCoreBindings.h b/Engine/Plugins/Experimental/PlainProps/Source/Public/PlainPropsUeCoreBindings.h
--- a/Engine/Plugins/Experimental/PlainProps/Source/Public/PlainPropsUeCoreBindings.h
+++ b/Engine/Plugins/Experimental/PlainProps/Source/Public/PlainPropsUeCoreBindings.h
@@ -371,16 +371,24 @@ struct FTransformIds
 
 struct FTransformBinding : public ICustomBinding
 {
-	PLAINPROPS_API void			Save(FMemberBuilder& Dst, const FTransform& Src, const FTransform* Default, const FTransformIds& Ids) const;
-	PLAINPROPS_API void			Load(FTransform& Dst, FStructView Src, ECustomLoadMethod Method, const FLoadBatch& Batch, const FTransformIds& Ids) const;
+	using Type = FTransform;
+	static constexpr EMemberPresence Occupancy = EMemberPresence::AllowSparse;
+
+
+
+	const FTransformIds& MemberIds;
+	FTransformBinding(const FTransformIds& Ids) : MemberIds(Ids) {}
+	PLAINPROPS_API void			Save(FMemberBuilder& Dst, const FTransform& Src, const FTransform* Default) const;
+	PLAINPROPS_API void			Load(FTransform& Dst, FStructView Src, ECustomLoadMethod Method, const FLoadBatch& Batch) const;
 	PLAINPROPS_API virtual bool	DiffStruct(const void* StructA, const void* StructB) const override;
 };
 
 template<class Ids>
 struct TTransformBinding final : public FTransformBinding
 {
-	using Type = FTransform;
-	static constexpr EMemberPresence Occupancy = EMemberPresence::AllowSparse;
+	TTransformBinding()
+	: FTransformBinding(FTransformIds::Get<Ids>()) 
+	{}
 
 	static TConstArrayView<FMemberId> GetMemberIds()
 	{
@@ -389,12 +397,12 @@ struct TTransformBinding final : public FTransformBinding
 
 	virtual void SaveStruct(FMemberBuilder& Dst, const void* Src, const void* Default, const FDebugIds& Debug) override
 	{
-		Save(Dst, *static_cast<const FTransform*>(Src), static_cast<const FTransform*>(Default), FTransformIds::Get<Ids>());
+		Save(Dst, *static_cast<const FTransform*>(Src), static_cast<const FTransform*>(Default));
 	}
 
 	virtual void LoadStruct(void* Dst, FStructView Src, ECustomLoadMethod Method, const FLoadBatch& Batch) const override
 	{
-		Load(*static_cast<FTransform*>(Dst), Src, Method, Batch, FTransformIds::Get<Ids>());
+		Load(*static_cast<FTransform*>(Dst), Src, Method, Batch);
 	}
 };
 
diff --git a/Engine/Plugins/Experimental/PlainProps/Source/Tests/PlainPropsUeCoreTests.cpp b/Engine/Plugins/Experimental/PlainProps/Source/Tests/PlainPropsUeCoreTests.cpp
--- a/Engine/Plugins/Experimental/PlainProps/Source/Tests/PlainPropsUeCoreTests.cpp
+++ b/Engine/Plugins/Experimental/PlainProps/Source/Tests/PlainPropsUeCoreTests.cpp
@@ -778,14 +778,9 @@ TEST_CASE_NAMED(FPlainPropsUeCoreTest, "System::Core::Serialization::PlainProps:
 			{
 				CHECK(Batch.Load<FTransform>().Equals(FTransform(), 0.0));
 
-				FTransform TranslateOnly;
-				TranslateOnly.SetTranslation(FVector(10, 20, 30));
-				TranslateOnly.SetRotation(FQuat(1, 2, 3, 4));
-				TranslateOnly.SetScale3D(FVector(5, 6, 7));
-				Batch.LoadInto(TranslateOnly);
-				CHECK(TranslateOnly.GetTranslation() == FVector::UnitY());
-				CHECK(TranslateOnly.GetRotation() == FQuat(1, 2, 3, 4));
-				CHECK(TranslateOnly.GetScale3D() == FVector(5, 6, 7));
+				FTransform TranslateY(				FQuat(1, 2, 3, 4), FVector(5, 5, 5), FVector(6, 7, 8));
+				Batch.LoadInto(TranslateY);
+				CHECK(TranslateY.Equals(FTransform(	FQuat(1, 2, 3, 4), FVector::UnitY(), FVector(6, 7, 8)), 0.0));
 			});
 	}
 
