PlainProps improve  custom binding API ergonomics

#rnx

diff --git a/Engine/Plugins/Experimental/PlainProps/Source/Private/PlainPropsSave.cpp b/Engine/Plugins/Experimental/PlainProps/Source/Private/PlainPropsSave.cpp
--- a/Engine/Plugins/Experimental/PlainProps/Source/Private/PlainPropsSave.cpp
+++ b/Engine/Plugins/Experimental/PlainProps/Source/Private/PlainPropsSave.cpp
@@ -269,7 +269,7 @@ FBuiltStructPtr SaveStruct(const void* Struct, FStructSchemaId Id, const FSaveCo
 	FMemberBuilder Out;
 	if (ICustomBinding* Custom = Ctx.Customs.FindStruct(Id))
 	{
-		Custom->SaveStruct(Out, Struct, nullptr, Ctx.Declarations.GetDebug());
+		Custom->SaveStruct(Out, Struct, nullptr, Ctx);
 	}
 	else for (FMemberVisitor It(Ctx.Schemas.GetStruct(Id)); It.HasMore(); )
 	{
@@ -485,7 +485,7 @@ FBuiltStructPtr SaveStructDelta(const void* Struct, const void* Default, FStruct
 	{
 		if (Custom->DiffStruct(Struct, Default))
 		{
-			Custom->SaveStruct(Out, Struct, Default, Ctx.Declarations.GetDebug());
+			Custom->SaveStruct(Out, Struct, Default, Ctx);
 		}
 	}
 	else for (FMemberVisitor It(Ctx.Schemas.GetStruct(Id)); It.HasMore(); )
diff --git a/Engine/Plugins/Experimental/PlainProps/Source/Private/PlainPropsUeCoreBindings.cpp b/Engine/Plugins/Experimental/PlainProps/Source/Private/PlainPropsUeCoreBindings.cpp
--- a/Engine/Plugins/Experimental/PlainProps/Source/Private/PlainPropsUeCoreBindings.cpp
+++ b/Engine/Plugins/Experimental/PlainProps/Source/Private/PlainPropsUeCoreBindings.cpp
@@ -7,7 +7,7 @@
 namespace PlainProps::UE
 {
 
-void FTransformBinding::Save(FMemberBuilder& Dst, const FTransform& Src, const FTransform* Default) const
+void FTransformBinding::Save(FMemberBuilder& Dst, const FTransform& Src, const FTransform* Default, const FSaveContext& Context) const
 {
 	static_assert(std::is_same_v<decltype(FTransform().GetTranslation().X), double>);
 
@@ -19,38 +19,38 @@ void FTransformBinding::Save(FMemberBuilder& Dst, const FTransform& Src, const F
 	{
 		if (T != Default->GetTranslation())
 		{
-			Dst.Add(MemberIds.Translate[0], T.X);
-			Dst.Add(MemberIds.Translate[1], T.Y);
-			Dst.Add(MemberIds.Translate[2], T.Z);
+			Dst.Add(MemberIds[(uint8)EMember::TranslateX], T.X);
+			Dst.Add(MemberIds[(uint8)EMember::TranslateY], T.Y);
+			Dst.Add(MemberIds[(uint8)EMember::TranslateZ], T.Z);
 		}
 
 		if (R != Default->GetRotation())
 		{
-			Dst.Add(MemberIds.Rotate[0], R.X);
-			Dst.Add(MemberIds.Rotate[1], R.Y);
-			Dst.Add(MemberIds.Rotate[2], R.Z);
-			Dst.Add(MemberIds.Rotate[3], R.W);
+			Dst.Add(MemberIds[(uint8)EMember::RotateX], R.X);
+			Dst.Add(MemberIds[(uint8)EMember::RotateY], R.Y);
+			Dst.Add(MemberIds[(uint8)EMember::RotateZ], R.Z);
+			Dst.Add(MemberIds[(uint8)EMember::RotateW], R.W);
 		}
 
 		if (S != Default->GetScale3D())
 		{
-			Dst.Add(MemberIds.Scale[0], S.X);
-			Dst.Add(MemberIds.Scale[1], S.Y);
-			Dst.Add(MemberIds.Scale[2], S.Z);
+			Dst.Add(MemberIds[(uint8)EMember::ScaleX], S.X);
+			Dst.Add(MemberIds[(uint8)EMember::ScaleY], S.Y);
+			Dst.Add(MemberIds[(uint8)EMember::ScaleZ], S.Z);
 		}
 	}
 	else
 	{
-		Dst.Add(MemberIds.Translate[0],	T.X);
-		Dst.Add(MemberIds.Translate[1],	T.Y);
-		Dst.Add(MemberIds.Translate[2],	T.Z);
-		Dst.Add(MemberIds.Rotate[0],		R.X);
-		Dst.Add(MemberIds.Rotate[1],		R.Y);
-		Dst.Add(MemberIds.Rotate[2],		R.Z);
-		Dst.Add(MemberIds.Rotate[3],		R.W);
-		Dst.Add(MemberIds.Scale[0],		S.X);
-		Dst.Add(MemberIds.Scale[1],		S.Y);
-		Dst.Add(MemberIds.Scale[2],		S.Z);
+		Dst.Add(MemberIds[(uint8)EMember::TranslateX],	T.X);
+		Dst.Add(MemberIds[(uint8)EMember::TranslateY],	T.Y);
+		Dst.Add(MemberIds[(uint8)EMember::TranslateZ],	T.Z);
+		Dst.Add(MemberIds[(uint8)EMember::RotateX], R.X);
+		Dst.Add(MemberIds[(uint8)EMember::RotateY], R.Y);
+		Dst.Add(MemberIds[(uint8)EMember::RotateZ], R.Z);
+		Dst.Add(MemberIds[(uint8)EMember::RotateW], R.W);
+		Dst.Add(MemberIds[(uint8)EMember::ScaleX],	S.X);
+		Dst.Add(MemberIds[(uint8)EMember::ScaleY],	S.Y);
+		Dst.Add(MemberIds[(uint8)EMember::ScaleZ],	S.Z);
 	}
 }
 
@@ -70,7 +70,7 @@ void FTransformBinding::Load(FTransform& Dst, FStructView Src, ECustomLoadMethod
 		return;
 	}
 
-	if (Members.PeekNameUnchecked() == MemberIds.Translate[0])
+	if (Members.PeekNameUnchecked() == MemberIds[(uint8)EMember::TranslateX])
 	{
 		FVector Translation;
 		Members.GrabLeaves(&Translation.X, 3);
@@ -82,7 +82,7 @@ void FTransformBinding::Load(FTransform& Dst, FStructView Src, ECustomLoadMethod
 		}
 	}
 
-	if (Members.PeekNameUnchecked() == MemberIds.Rotate[0])
+	if (Members.PeekNameUnchecked() == MemberIds[(uint8)EMember::RotateX])
 	{
 		FQuat Rotation;
 		Members.GrabLeaves(&Rotation.X, 4);
@@ -94,18 +94,13 @@ void FTransformBinding::Load(FTransform& Dst, FStructView Src, ECustomLoadMethod
 		}
 	}
 
-	checkSlow(Members.PeekNameUnchecked() == MemberIds.Scale[0]);
+	checkSlow(Members.PeekNameUnchecked() == MemberIds[(uint8)EMember::ScaleX]);
 	FVector Scale;
 	Members.GrabLeaves(&Scale.X, 3);
 	Dst.SetScale3D(Scale);
 	checkSlow(!Members.HasMore());
 }
 
-bool FTransformBinding::DiffStruct(const void* StructA, const void* StructB) const
-{
-	return !static_cast<const FTransform*>(StructA)->Equals(*static_cast<const FTransform*>(StructB), 0.0);
-}
-
 } // namespace PlainProps::UE
 
 //////////////////////////////////////////////////////////////////////////
diff --git a/Engine/Plugins/Experimental/PlainProps/Source/Public/PlainPropsBind.h b/Engine/Plugins/Experimental/PlainProps/Source/Public/PlainPropsBind.h
--- a/Engine/Plugins/Experimental/PlainProps/Source/Public/PlainPropsBind.h
+++ b/Engine/Plugins/Experimental/PlainProps/Source/Public/PlainPropsBind.h
@@ -22,6 +22,7 @@ struct FSchemaBatch;
 class FScratchAllocator;
 class FStructBinding;
 class FRangeBinding;
+struct FSaveContext;
 struct FTypedRange;
 class IItemRangeBinding;
 template<class T> class TIdIndexer;
@@ -238,7 +239,7 @@ class ICustomBinding
 {
 public:
 	virtual ~ICustomBinding() {}
-	virtual void				SaveStruct(FMemberBuilder& Dst, const void* Src, const void* Default, const FDebugIds& Debug) = 0;
+	virtual void				SaveStruct(FMemberBuilder& Dst, const void* Src, const void* Default, const FSaveContext& Ctx) = 0;
 	virtual void				LoadStruct(void* Dst, FStructView Src, ECustomLoadMethod Method, const FLoadBatch& Batch) const = 0;
 	virtual bool				DiffStruct(const void* StructA, const void* StructB) const = 0;
 };
@@ -269,11 +270,11 @@ private:
 	const FDebugIds&						Debug;
 };
 
-template<typename T, typename Ids>
+template<typename T>
 struct TCustomBind{ using Type = void; };
 
-template<typename T, typename Ids>
-using CustomBind = typename TCustomBind<T, Ids>::Type;
+template<typename T>
+using CustomBind = typename TCustomBind<T>::Type;
 
 ////////////////////////////////////////////////////////////////////////////////////////////////
 
@@ -569,7 +570,7 @@ FMemberBindType BindMemberLeaf(FOptionalSchemaId& OutSchema)
 template<typename CustomBinding, class Runtime>
 FStructSchemaId BindCustomStructOnce()
 {
-	struct FBinding
+	struct FBinding : CustomBinding
 	{
 		using Type = typename CustomBinding::Type;
 		using Ids = typename Runtime::Ids;
@@ -578,8 +579,9 @@ FStructSchemaId BindCustomStructOnce()
 		{
 			FTypeId Name = Ids::IndexNativeType(CttiOf<Type>::Name);
 			Id = Ids::IndexStruct(Name);
-			Runtime::GetTypes().DeclareStruct(Id, Name, CustomBinding::GetMemberIds(), CustomBinding::Occupancy);
-			Runtime::GetCustoms().BindStruct(Id, Instance);
+			CustomBinding::template InitIds<Ids>();
+			Runtime::GetTypes().DeclareStruct(Id, Name, CustomBinding::MemberIds, CustomBinding::Occupancy);
+			Runtime::GetCustoms().BindStruct(Id, *this);
 		}
 
 		~FBinding()
@@ -588,7 +590,21 @@ FStructSchemaId BindCustomStructOnce()
 			Runtime::GetTypes().DropStruct(Id);
 		}
 
-		CustomBinding Instance;
+		virtual void SaveStruct(FMemberBuilder& Dst, const void* Src, const void* Default, const FSaveContext& Ctx) override
+		{
+			CustomBinding::Save(Dst, *static_cast<const Type*>(Src), static_cast<const Type*>(Default), Ctx);
+		}
+
+		virtual void LoadStruct(void* Dst, FStructView Src, ECustomLoadMethod Method, const FLoadBatch& Batch) const override
+		{
+			CustomBinding::Load(*static_cast<Type*>(Dst), Src, Method, Batch);
+		}
+
+		virtual bool DiffStruct(const void* A, const void* B) const override
+		{
+			return CustomBinding::Diff(*static_cast<const Type*>(A), *static_cast<const Type*>(B));
+		}
+
 		FStructSchemaId Id;
 	};
 
diff --git a/Engine/Plugins/Experimental/PlainProps/Source/Public/PlainPropsUeCoreBindings.h b/Engine/Plugins/Experimental/PlainProps/Source/Public/PlainPropsUeCoreBindings.h
--- a/Engine/Plugins/Experimental/PlainProps/Source/Public/PlainPropsUeCoreBindings.h
+++ b/Engine/Plugins/Experimental/PlainProps/Source/Public/PlainPropsUeCoreBindings.h
@@ -3,7 +3,9 @@
 #pragma once
 
 #include "Containers/Set.h"
-#include "Math/MathFwd.h"
+#include "Math/Quat.h"
+#include "Math/Vector.h"
+#include "Math/Transform.h"
 #include "PlainPropsBind.h"
 #include "PlainPropsLoad.h"
 #include "PlainPropsRead.h"
@@ -11,6 +13,17 @@
 #include "UObject/NameTypes.h"
 
 
+PP_NAME_STRUCT(, FName);
+
+namespace UE::Math
+{
+PP_REFLECT_STRUCT(, FVector, void, X, Y, Z);
+PP_REFLECT_STRUCT(, FVector4, void, X, Y, Z, W);
+PP_REFLECT_STRUCT(, FQuat, void, X, Y, Z, W);
+PP_NAME_STRUCT(, FTransform);
+}
+
+
 namespace PlainProps::UE
 {
 
@@ -344,66 +357,44 @@ struct TSetBinding : public IItemRangeBinding
 
 //////////////////////////////////////////////////////////////////////////
 
-struct FTransformIds
-{
-	FMemberId Translate[3];
-	FMemberId Rotate[4];
-	FMemberId Scale[3];
-
-	TConstArrayView<FMemberId> GetAll() const
-	{
-		static_assert(offsetof(FTransformIds, Rotate)	== offsetof(FTransformIds, Translate) + sizeof(Translate));
-		static_assert(offsetof(FTransformIds, Scale)	== offsetof(FTransformIds, Rotate) + sizeof(Rotate));
-		return MakeArrayView(&Translate[0], 3 + 4 + 3);
-	}
-	
-	template<class Ids>	
-	static const FTransformIds& Get()
-	{
-		static FTransformIds Out = {	//Ids::IndexMember("Mask"),
-			{ Ids::IndexMember("TranslateX"), Ids::IndexMember("TranslateY"), Ids::IndexMember("TranslateZ") },
-			{ Ids::IndexMember("RotateX"), Ids::IndexMember("RotateY"),	Ids::IndexMember("RotateZ"), Ids::IndexMember("RotateW") },
-			{ Ids::IndexMember("ScaleX"), Ids::IndexMember("ScaleY"), Ids::IndexMember("ScaleZ") }};
-
-		return Out;
-	}
-};
-
+//TODO: macroify, e.g PP_CUSTOM_BIND(PLAINPROPS_API, FTransform, Transform, Translate, Rotate, Scale)
 struct FTransformBinding : public ICustomBinding
 {
 	using Type = FTransform;
-	static constexpr EMemberPresence Occupancy = EMemberPresence::AllowSparse;
-
-
-
-	const FTransformIds& MemberIds;
-	FTransformBinding(const FTransformIds& Ids) : MemberIds(Ids) {}
-	PLAINPROPS_API void			Save(FMemberBuilder& Dst, const FTransform& Src, const FTransform* Default) const;
-	PLAINPROPS_API void			Load(FTransform& Dst, FStructView Src, ECustomLoadMethod Method, const FLoadBatch& Batch) const;
-	PLAINPROPS_API virtual bool	DiffStruct(const void* StructA, const void* StructB) const override;
-};
-
-template<class Ids>
-struct TTransformBinding final : public FTransformBinding
-{
-	TTransformBinding()
-	: FTransformBinding(FTransformIds::Get<Ids>()) 
-	{}
-
-	static TConstArrayView<FMemberId> GetMemberIds()
+	inline static constexpr EMemberPresence Occupancy = EMemberPresence::AllowSparse;
+	//inline static constexpr uint8 NumMembers = 3;
+	//enum class EMember : uint8 { Translate, Rotate, Scale };
+	//FMemberId MemberIds[3];
+	//FStructSchemaId VectorId;
+	//FStructSchemaId QuatId;
+
+	enum class EMember : uint8 { TranslateX, TranslateY, TranslateZ, RotateX, RotateY, RotateZ, RotateW, ScaleX, ScaleY, ScaleZ };
+	FMemberId MemberIds[11];
+
+	template<typename Ids>
+	void InitIds() 
 	{
-		return FTransformIds::Get<Ids>().GetAll();
-	}
-
-	virtual void SaveStruct(FMemberBuilder& Dst, const void* Src, const void* Default, const FDebugIds& Debug) override
-	{
-		Save(Dst, *static_cast<const FTransform*>(Src), static_cast<const FTransform*>(Default));
+		//MemberIds[EMember::Translate] = Ids::IndexMember("Translate");
+		//MemberIds[EMember::Rotate] = Ids::IndexMember("Rotate");
+		//MemberIds[EMember::Scale] = Ids::IndexMember("Scale");
+		//VectorId = IndexNativeStruct<FVector, Ids>();
+		//QuatId = IndexNativeStruct<FVector, Ids>();
+		
+		MemberIds[(uint8)EMember::TranslateX] = Ids::IndexMember("TranslateX");
+		MemberIds[(uint8)EMember::TranslateY] = Ids::IndexMember("TranslateY");
+		MemberIds[(uint8)EMember::TranslateZ] = Ids::IndexMember("TranslateZ");
+		MemberIds[(uint8)EMember::RotateX] = Ids::IndexMember("RotateX");
+		MemberIds[(uint8)EMember::RotateY] = Ids::IndexMember("RotateY");
+		MemberIds[(uint8)EMember::RotateZ] = Ids::IndexMember("RotateZ");
+		MemberIds[(uint8)EMember::RotateW] = Ids::IndexMember("RotateW");
+		MemberIds[(uint8)EMember::ScaleX] = Ids::IndexMember("ScaleX");
+		MemberIds[(uint8)EMember::ScaleY] = Ids::IndexMember("ScaleY");
+		MemberIds[(uint8)EMember::ScaleZ] = Ids::IndexMember("ScaleZ");
 	}
 
-	virtual void LoadStruct(void* Dst, FStructView Src, ECustomLoadMethod Method, const FLoadBatch& Batch) const override
-	{
-		Load(*static_cast<FTransform*>(Dst), Src, Method, Batch);
-	}
+	PLAINPROPS_API void	Save(FMemberBuilder& Dst, const FTransform& Src, const FTransform* Default, const FSaveContext& Context) const;
+	PLAINPROPS_API void	Load(FTransform& Dst, FStructView Src, ECustomLoadMethod Method, const FLoadBatch& Batch) const;
+	inline static bool	Diff(const FTransform& A, const FTransform& B) { return !A.Equals(B, 0.0); }
 };
 
 //////////////////////////////////////////////////////////////////////////
@@ -438,7 +429,7 @@ struct TSetDeltaBinding : public ICustomBinding
 	
 	static TConstArrayView<FMemberId> GetMemberIds() { return FSetOps::Get<Ids>().All; }
 
-	virtual void SaveStruct(FMemberBuilder& Dst, const void* Src, const void* Default, const FDebugIds& Debug) override;
+	virtual void SaveStruct(FMemberBuilder& Dst, const void* Src, const void* Default, const FSaveContext& Ctx) override;
 
 	virtual void LoadStruct(void* Dst, FStructView Src, ECustomLoadMethod Method, const FLoadBatch& Batch) const override
 	{
@@ -622,9 +613,6 @@ struct TSetDeltaBinding : public ICustomBinding
 }
 
 
-PP_NAME_STRUCT(, FName);
-
-namespace UE::Math { PP_NAME_STRUCT(, FTransform); }
 
 
 namespace PlainProps
@@ -657,9 +645,9 @@ struct TRangeBind<TSet<T>>
 	using Type = UE::TSetBinding<T>;
 };
 
-template<typename Ids>
-struct TCustomBind<FTransform, Ids>
+template<>
+struct TCustomBind<FTransform>
 {
-	using Type = UE::TTransformBinding<Ids>;
+	using Type = UE::FTransformBinding;
 };
 }
\ No newline at end of file
diff --git a/Engine/Plugins/Experimental/PlainProps/Source/Tests/PlainPropsUeCoreTests.cpp b/Engine/Plugins/Experimental/PlainProps/Source/Tests/PlainPropsUeCoreTests.cpp
--- a/Engine/Plugins/Experimental/PlainProps/Source/Tests/PlainPropsUeCoreTests.cpp
+++ b/Engine/Plugins/Experimental/PlainProps/Source/Tests/PlainPropsUeCoreTests.cpp
@@ -43,18 +43,10 @@ struct FIds
 
 ////////////////////////////////////////////////////////////////////////////////////////////////
 
-template<class T>
-struct TCustomBindings
-{
-	using Type = CustomBind<T, FIds>;
-};
-
-////////////////////////////////////////////////////////////////////////////////////////////////
-
 struct FDefaultRuntime
 {
 	using Ids = FIds;
-	template<class T> using CustomBindings = TCustomBindings<T>;
+	template<class T> using CustomBindings = TCustomBind<T>;
 
 	static FDeclarations&			GetTypes()			{ return GTypes; }
 	static FSchemaBindings&			GetSchemas()		{ return GSchemas; }
@@ -140,7 +132,7 @@ struct FTestCustomBinding : public ICustomBinding
 
 struct FNameBinding : public FTestCustomBinding
 {
-	virtual void SaveStruct(FMemberBuilder& Dst, const void* Src, const void*, const FDebugIds& Debug) override
+	virtual void SaveStruct(FMemberBuilder& Dst, const void* Src, const void*, const FSaveContext& Ctx) override
 	{
 		FSetElementId Idx = Names.Add(*static_cast<const FName*>(Src));
 		Dst.Add(Declaration.Idx, Idx.AsInteger());
@@ -546,7 +538,7 @@ static bool operator==(const FDelta& A, const FDelta& B) { return A.A == B.A &&
 //		return MakeArrayView(Ids, 2);
 //	}
 //
-//	virtual void SaveStruct(FMemberBuilder& Dst, const void* Src, const void*, const FDebugIds&) const override
+//	virtual void SaveStruct(FMemberBuilder& Dst, const void* Src, const void*, const FSaveContext&) const override
 //	{
 //		if (const FObject* Object = reinterpret_cast<const FObject*>(Src))
 //		{
@@ -762,7 +754,7 @@ TEST_CASE_NAMED(FPlainPropsUeCoreTest, "System::Core::Serialization::PlainProps:
 
 	SECTION("Transform")
 	{
-		BindCustomStructOnce<TTransformBinding<FIds>, FDefaultRuntime>();
+		BindCustomStructOnce<FTransformBinding, FDefaultRuntime>();
 
 		Run([](FBatchSaver& Batch)
 			{
