PlainProps misc opts, improved custom bindings ergomics, added draft FDenseMemberBuilder and saving FTransform as 3 structs instead of 11 doubles

#rnx

diff --git a/Engine/Plugins/Experimental/PlainProps/Source/Private/PlainPropsBuild.cpp b/Engine/Plugins/Experimental/PlainProps/Source/Private/PlainPropsBuild.cpp
--- a/Engine/Plugins/Experimental/PlainProps/Source/Private/PlainPropsBuild.cpp
+++ b/Engine/Plugins/Experimental/PlainProps/Source/Private/PlainPropsBuild.cpp
@@ -138,39 +138,58 @@ FBuiltStructPtr FMemberBuilder::BuildAndReset(FScratchAllocator& Scratch, const
 	checkf(!(Declared.Super && Declared.Occupancy == EMemberPresence::RequireAll),
 		TEXT("Requiring sub structs to be dense isn't implemented"));
 #if DO_CHECK
-	int32 OrderIdx = 0;
-	TConstArrayView<FMemberId> Order = Declared.GetMemberOrder();
-	for (FBuiltMember& Member : Members)
+	// Verify members were added in declared order
+	if (int32 Num = Members.Num())
 	{
-		if (Declared.Super && IsSuper(Member.Schema.Type) && &Member == &Members[0])
+		int32 OrderIdx = 0;
+		TConstArrayView<FMemberId> Order = Declared.GetMemberOrder();
+		int32 SkipSuper = Declared.Super && IsSuper(Members[0].Schema.Type);
+		for (FBuiltMember* It = Members.GetData() + SkipSuper, *End = Members.GetData() + Num; It != End; ++It)
 		{
-			continue; // Skip undeclared generated super struct
+			FBuiltMember& Member = *It;
+			for (; OrderIdx < Order.Num() && Order[OrderIdx] != Member.Name; ++OrderIdx)
+			{}	
+			checkf(OrderIdx < Order.Num(), TEXT("Member '%s' in '%s' %s"), *Debug.Print(Member.Name), *Debug.Print(Declared.Type),
+					Order.Contains(Member.Name) ? TEXT("appeared in non-declared order") : TEXT("is undeclared"));
+			++OrderIdx;
 		}
-
-		for (; OrderIdx < Order.Num() && Order[OrderIdx] != Member.Name; ++OrderIdx)
-		{}	
-		checkf(OrderIdx < Order.Num(), TEXT("Member '%s' in '%s' %s"), *Debug.Print(Member.Name), *Debug.Print(Declared.Type),
-				Order.Contains(Member.Name) ? TEXT("appeared in non-declared order") : TEXT("is undeclared"));
-		++OrderIdx;
 	}
 #endif
 
 	uint32 Num = static_cast<uint32>(Members.Num());
 	SIZE_T NumBytes = sizeof(FBuiltStruct) + Num * sizeof(FBuiltMember);
-	FBuiltStruct* Out = reinterpret_cast<FBuiltStruct*>(Scratch.AllocateZeroed(NumBytes, alignof(FBuiltStruct)));
+	FBuiltStruct* Out = reinterpret_cast<FBuiltStruct*>(Scratch.Allocate(NumBytes, alignof(FBuiltStruct)));
 	Out->NumMembers = IntCastChecked<uint16>(Num);
-	for (FBuiltMember& Member : Members)
-	{
-		Out->Members[&Member - &Members[0]] = MoveTemp(Member);
-	}
+	FMemory::Memcpy(Out->Members, Members.GetData(), Members.NumBytes());
 
 	Members.Reset();
 
-	//GLiveStructsFoo.Add(Out);
-
 	return FBuiltStructPtr(Out);
 }
 
+//////////////////////////////////////////////////////////////////////////
+
+FBuiltStructPtr FDenseMemberBuilder::BuildHomo(const FStructDeclaration& Declaration, FMemberType Leaf, TConstArrayView<FBuiltValue> Values) const
+{
+	check(Declaration.NumMembers == Values.Num());
+
+	FMemberSchema Schema = {Leaf, Leaf};
+	const uint32 Num = static_cast<uint32>(Values.Num());
+	SIZE_T NumBytes = sizeof(FBuiltStruct) + Num * sizeof(FBuiltMember);
+	FBuiltStruct* Out = reinterpret_cast<FBuiltStruct*>(Scratch.Allocate(NumBytes, alignof(FBuiltStruct)));
+	Out->NumMembers = static_cast<uint16>(Num);
+
+	const FMemberId* Names = Declaration.MemberOrder; 
+	for (uint32 Idx = 0; Idx < Num; ++Idx)
+	{
+		new (Out->Members + Idx) FBuiltMember(Names[Idx], Schema, Values[Idx]);
+	}
+
+	return Out;
+}
+
+//////////////////////////////////////////////////////////////////////////
+
 template<typename IntType, typename FloatType>
 inline IntType CheckFiniteBitCast(FloatType Value)
 {
diff --git a/Engine/Plugins/Experimental/PlainProps/Source/Private/PlainPropsLoad.cpp b/Engine/Plugins/Experimental/PlainProps/Source/Private/PlainPropsLoad.cpp
--- a/Engine/Plugins/Experimental/PlainProps/Source/Private/PlainPropsLoad.cpp
+++ b/Engine/Plugins/Experimental/PlainProps/Source/Private/PlainPropsLoad.cpp
@@ -615,7 +615,7 @@ void LoadStruct(void* Dst, FByteReader Src, FStructSchemaId Id, const FLoadBatch
 	else
 	{
 		FStructSchemaHandle ReadSchema{Id, Batch.ReadId};
-		Plan.AsCustom().LoadStruct(Dst, { ReadSchema, Src }, ECustomLoadMethod::Assign, Batch);
+		Plan.AsCustom().LoadCustom(Dst, { ReadSchema, Src }, ECustomLoadMethod::Assign, Batch);
 	}
 }
 
@@ -637,7 +637,7 @@ void ConstructAndLoadStruct(void* Dst, FByteReader Src, FStructSchemaId Id, cons
 	else
 	{
 		FStructSchemaHandle ReadSchema{Id, Batch.ReadId};
-		Plan.AsCustom().LoadStruct(Dst, { ReadSchema, Src }, ECustomLoadMethod::Construct, Batch);
+		Plan.AsCustom().LoadCustom(Dst, { ReadSchema, Src }, ECustomLoadMethod::Construct, Batch);
 	}
 }
 
diff --git a/Engine/Plugins/Experimental/PlainProps/Source/Private/PlainPropsSave.cpp b/Engine/Plugins/Experimental/PlainProps/Source/Private/PlainPropsSave.cpp
--- a/Engine/Plugins/Experimental/PlainProps/Source/Private/PlainPropsSave.cpp
+++ b/Engine/Plugins/Experimental/PlainProps/Source/Private/PlainPropsSave.cpp
@@ -269,7 +269,7 @@ FBuiltStructPtr SaveStruct(const void* Struct, FStructSchemaId Id, const FSaveCo
 	FMemberBuilder Out;
 	if (ICustomBinding* Custom = Ctx.Customs.FindStruct(Id))
 	{
-		Custom->SaveStruct(Out, Struct, nullptr, Ctx);
+		Custom->SaveCustom(Out, Struct, nullptr, Ctx);
 	}
 	else for (FMemberVisitor It(Ctx.Schemas.GetStruct(Id)); It.HasMore(); )
 	{
@@ -428,7 +428,7 @@ static bool DiffItem(const uint8* A, const uint8* B, const FSaveContext& Ctx, FS
 {
 	if (ICustomBinding* Custom = Ctx.Customs.FindStruct(Id))
 	{
-		return Custom->DiffStruct(A, B);
+		return Custom->DiffCustom(A, B);
 	}
 	
 	bool bOut = false;
@@ -483,9 +483,9 @@ FBuiltStructPtr SaveStructDelta(const void* Struct, const void* Default, FStruct
 	FMemberBuilder Out;
 	if (ICustomBinding* Custom = Ctx.Customs.FindStruct(Id))
 	{
-		if (Custom->DiffStruct(Struct, Default))
+		if (Custom->DiffCustom(Struct, Default))
 		{
-			Custom->SaveStruct(Out, Struct, Default, Ctx);
+			Custom->SaveCustom(Out, Struct, Default, Ctx);
 		}
 	}
 	else for (FMemberVisitor It(Ctx.Schemas.GetStruct(Id)); It.HasMore(); )
diff --git a/Engine/Plugins/Experimental/PlainProps/Source/Private/PlainPropsUeCoreBindings.cpp b/Engine/Plugins/Experimental/PlainProps/Source/Private/PlainPropsUeCoreBindings.cpp
--- a/Engine/Plugins/Experimental/PlainProps/Source/Private/PlainPropsUeCoreBindings.cpp
+++ b/Engine/Plugins/Experimental/PlainProps/Source/Private/PlainPropsUeCoreBindings.cpp
@@ -2,14 +2,19 @@
 
 #include "PlainPropsUeCoreBindings.h"
 #include "PlainPropsBuild.h"
+#include "PlainPropsSave.h"
 #include "Math/Transform.h"
 
 namespace PlainProps::UE
 {
 
-void FTransformBinding::Save(FMemberBuilder& Dst, const FTransform& Src, const FTransform* Default, const FSaveContext& Context) const
+void FTransformBinding::Save(FMemberBuilder& Dst, const FTransform& Src, const FTransform* Default, const FSaveContext& Ctx) const
 {
 	static_assert(std::is_same_v<decltype(FTransform().GetTranslation().X), double>);
+	
+	const FStructDeclaration& VectorDecl = Ctx.Declarations.Get(VectorId);
+	const FStructDeclaration& QuatDecl = Ctx.Declarations.Get(QuatId);
+	FDenseMemberBuilder Inner = { Ctx.Scratch, Ctx.Declarations.GetDebug() };
 
 	FVector T = Src.GetTranslation();
 	FQuat R = Src.GetRotation();
@@ -19,41 +24,37 @@ void FTransformBinding::Save(FMemberBuilder& Dst, const FTransform& Src, const F
 	{
 		if (T != Default->GetTranslation())
 		{
-			Dst.Add(MemberIds[(uint8)EMember::TranslateX], T.X);
-			Dst.Add(MemberIds[(uint8)EMember::TranslateY], T.Y);
-			Dst.Add(MemberIds[(uint8)EMember::TranslateZ], T.Z);
+			Dst.AddStruct(MemberIds[(uint8)EMember::Translate], VectorId, Inner.BuildHomogeneous(VectorDecl, T.X, T.Y, T.Z));
 		}
 
 		if (R != Default->GetRotation())
 		{
-			Dst.Add(MemberIds[(uint8)EMember::RotateX], R.X);
-			Dst.Add(MemberIds[(uint8)EMember::RotateY], R.Y);
-			Dst.Add(MemberIds[(uint8)EMember::RotateZ], R.Z);
-			Dst.Add(MemberIds[(uint8)EMember::RotateW], R.W);
+			Dst.AddStruct(MemberIds[(uint8)EMember::Rotate], QuatId, Inner.BuildHomogeneous(QuatDecl, R.X, R.Y, R.Z, R.W));
 		}
 
 		if (S != Default->GetScale3D())
 		{
-			Dst.Add(MemberIds[(uint8)EMember::ScaleX], S.X);
-			Dst.Add(MemberIds[(uint8)EMember::ScaleY], S.Y);
-			Dst.Add(MemberIds[(uint8)EMember::ScaleZ], S.Z);
+			Dst.AddStruct(MemberIds[(uint8)EMember::Scale], VectorId, Inner.BuildHomogeneous(VectorDecl, S.X, S.Y, S.Z));
 		}
 	}
 	else
-	{
-		Dst.Add(MemberIds[(uint8)EMember::TranslateX],	T.X);
-		Dst.Add(MemberIds[(uint8)EMember::TranslateY],	T.Y);
-		Dst.Add(MemberIds[(uint8)EMember::TranslateZ],	T.Z);
-		Dst.Add(MemberIds[(uint8)EMember::RotateX], R.X);
-		Dst.Add(MemberIds[(uint8)EMember::RotateY], R.Y);
-		Dst.Add(MemberIds[(uint8)EMember::RotateZ], R.Z);
-		Dst.Add(MemberIds[(uint8)EMember::RotateW], R.W);
-		Dst.Add(MemberIds[(uint8)EMember::ScaleX],	S.X);
-		Dst.Add(MemberIds[(uint8)EMember::ScaleY],	S.Y);
-		Dst.Add(MemberIds[(uint8)EMember::ScaleZ],	S.Z);
+	{	
+		Dst.AddStruct(MemberIds[(uint8)EMember::Translate], VectorId, Inner.BuildHomogeneous(VectorDecl, T.X, T.Y, T.Z));		
+		Dst.AddStruct(MemberIds[(uint8)EMember::Rotate], QuatId, Inner.BuildHomogeneous(QuatDecl, R.X, R.Y, R.Z, R.W));
+		Dst.AddStruct(MemberIds[(uint8)EMember::Scale], VectorId, Inner.BuildHomogeneous(VectorDecl, S.X, S.Y, S.Z));
 	}
 }
 
+template<typename T>
+T GrabAndMemcpy(FMemberReader& Members)
+{
+	T Out;
+	FStructView Struct = Members.GrabStruct();
+	Struct.Values.CheckSize(sizeof(T));
+	FMemory::Memcpy(&Out, Struct.Values.Peek(), sizeof(T));
+	return Out;
+}
+
 void FTransformBinding::Load(FTransform& Dst, FStructView Src, ECustomLoadMethod Method, const FLoadBatch& Batch) const
 {
 	static_assert(std::is_same_v<decltype(FTransform().GetTranslation().X), double>);
@@ -70,11 +71,9 @@ void FTransformBinding::Load(FTransform& Dst, FStructView Src, ECustomLoadMethod
 		return;
 	}
 
-	if (Members.PeekNameUnchecked() == MemberIds[(uint8)EMember::TranslateX])
+	if (Members.PeekNameUnchecked() == MemberIds[(uint8)EMember::Translate])
 	{
-		FVector Translation;
-		Members.GrabLeaves(&Translation.X, 3);
-		Dst.SetTranslation(Translation);
+		Dst.SetTranslation(GrabAndMemcpy<FVector>(Members));
 
 		if (!Members.HasMore())
 		{
@@ -82,11 +81,9 @@ void FTransformBinding::Load(FTransform& Dst, FStructView Src, ECustomLoadMethod
 		}
 	}
 
-	if (Members.PeekNameUnchecked() == MemberIds[(uint8)EMember::RotateX])
+	if (Members.PeekNameUnchecked() == MemberIds[(uint8)EMember::Rotate])
 	{
-		FQuat Rotation;
-		Members.GrabLeaves(&Rotation.X, 4);
-		Dst.SetRotation(Rotation);
+		Dst.SetRotation(GrabAndMemcpy<FQuat>(Members));
 
 		if (!Members.HasMore())
 		{
@@ -94,10 +91,8 @@ void FTransformBinding::Load(FTransform& Dst, FStructView Src, ECustomLoadMethod
 		}
 	}
 
-	checkSlow(Members.PeekNameUnchecked() == MemberIds[(uint8)EMember::ScaleX]);
-	FVector Scale;
-	Members.GrabLeaves(&Scale.X, 3);
-	Dst.SetScale3D(Scale);
+	checkSlow(Members.PeekNameUnchecked() == MemberIds[(uint8)EMember::Scale]);
+	Dst.SetScale3D(GrabAndMemcpy<FVector>(Members));
 	checkSlow(!Members.HasMore());
 }
 
diff --git a/Engine/Plugins/Experimental/PlainProps/Source/Public/PlainPropsBind.h b/Engine/Plugins/Experimental/PlainProps/Source/Public/PlainPropsBind.h
--- a/Engine/Plugins/Experimental/PlainProps/Source/Public/PlainPropsBind.h
+++ b/Engine/Plugins/Experimental/PlainProps/Source/Public/PlainPropsBind.h
@@ -239,9 +239,9 @@ class ICustomBinding
 {
 public:
 	virtual ~ICustomBinding() {}
-	virtual void				SaveStruct(FMemberBuilder& Dst, const void* Src, const void* Default, const FSaveContext& Ctx) = 0;
-	virtual void				LoadStruct(void* Dst, FStructView Src, ECustomLoadMethod Method, const FLoadBatch& Batch) const = 0;
-	virtual bool				DiffStruct(const void* StructA, const void* StructB) const = 0;
+	virtual void				SaveCustom(FMemberBuilder& Dst, const void* Src, const void* Default, const FSaveContext& Ctx) = 0;
+	virtual void				LoadCustom(void* Dst, FStructView Src, ECustomLoadMethod Method, const FLoadBatch& Batch) const = 0;
+	virtual bool				DiffCustom(const void* StructA, const void* StructB) const = 0;
 };
 
 class FCustomBindings
@@ -590,17 +590,17 @@ FStructSchemaId BindCustomStructOnce()
 			Runtime::GetTypes().DropStruct(Id);
 		}
 
-		virtual void SaveStruct(FMemberBuilder& Dst, const void* Src, const void* Default, const FSaveContext& Ctx) override
+		virtual void SaveCustom(FMemberBuilder& Dst, const void* Src, const void* Default, const FSaveContext& Ctx) override
 		{
 			CustomBinding::Save(Dst, *static_cast<const Type*>(Src), static_cast<const Type*>(Default), Ctx);
 		}
 
-		virtual void LoadStruct(void* Dst, FStructView Src, ECustomLoadMethod Method, const FLoadBatch& Batch) const override
+		virtual void LoadCustom(void* Dst, FStructView Src, ECustomLoadMethod Method, const FLoadBatch& Batch) const override
 		{
 			CustomBinding::Load(*static_cast<Type*>(Dst), Src, Method, Batch);
 		}
 
-		virtual bool DiffStruct(const void* A, const void* B) const override
+		virtual bool DiffCustom(const void* A, const void* B) const override
 		{
 			return CustomBinding::Diff(*static_cast<const Type*>(A), *static_cast<const Type*>(B));
 		}
diff --git a/Engine/Plugins/Experimental/PlainProps/Source/Public/PlainPropsBuild.h b/Engine/Plugins/Experimental/PlainProps/Source/Public/PlainPropsBuild.h
--- a/Engine/Plugins/Experimental/PlainProps/Source/Public/PlainPropsBuild.h
+++ b/Engine/Plugins/Experimental/PlainProps/Source/Public/PlainPropsBuild.h
@@ -204,14 +204,12 @@ struct FBuiltMember
 	FBuiltMember(FMemberId Name, FUnpackedLeafType Leaf, FOptionalEnumSchemaId Schema, uint64 Value);
 	FBuiltMember(FMemberId Name, FTypedRange Range);
 	FBuiltMember(FMemberId Name, FStructSchemaId Schema, FBuiltStructPtr Value);
+	FBuiltMember(FOptionalMemberId N, FMemberSchema S, FBuiltValue V) : Name(N), Schema(MoveTemp(S)), Value(V) {}
 	static FBuiltMember MakeSuper(FStructSchemaId Schema, FBuiltStructPtr Value);
 
 	FOptionalMemberId		Name;
 	FMemberSchema			Schema;
 	FBuiltValue				Value;
-
-private:
-	FBuiltMember(FOptionalMemberId N, FMemberSchema&& S, FBuiltValue V) : Name(N), Schema(MoveTemp(S)), Value(V) {}
 };
 
 //////////////////////////////////////////////////////////////////////////
@@ -254,13 +252,28 @@ private:
 	
 	FBuiltMemberArray		Members;
 	
-//	PLAINPROPS_API uint8* AllocateLeafRange(FMemberId Name, FUnpackedLeafType Leaf, ERangeSizeType RangeMax, uint64 Num);
-	
 	//template<typename T>
 	//void NormalizeLeafRange(T*, uint64) {}
 	//PLAINPROPS_API void NormalizeLeafRange(float*, uint64 Num);
 	//PLAINPROPS_API void NormalizeLeafRange(double*, uint64 Num);
-	
+};
+
+// Rough API draft
+struct FDenseMemberBuilder
+{
+	FScratchAllocator& Scratch;
+	const FDebugIds& Debug;
+
+	template<typename T, typename... Ts>
+	[[nodiscard]] FBuiltStructPtr BuildHomogeneous(const FStructDeclaration& Declaration, T Head, Ts... Tail) const
+	{
+		// Todo: Handle enums, ranges and structs
+		FBuiltValue Values[] = { {.Leaf = ValueCast(Head)}, {.Leaf = (ValueCast(Tail))}...  };
+		return BuildHomo(Declaration, FMemberType(ReflectLeaf<T>.Pack()), Values);
+	}
+
+private:
+	[[nodiscard]] PLAINPROPS_API FBuiltStructPtr BuildHomo(const FStructDeclaration& Declaration, FMemberType Leaf, TConstArrayView<FBuiltValue> Values) const;
 };
 
 // Helper class for building struct ranges
diff --git a/Engine/Plugins/Experimental/PlainProps/Source/Public/PlainPropsRead.h b/Engine/Plugins/Experimental/PlainProps/Source/Public/PlainPropsRead.h
--- a/Engine/Plugins/Experimental/PlainProps/Source/Public/PlainPropsRead.h
+++ b/Engine/Plugins/Experimental/PlainProps/Source/Public/PlainPropsRead.h
@@ -90,7 +90,7 @@ public:
 		uint32 NumBytesRead;
 		uint64 Out = ReadVarUInt(It, NumBytesRead);
 		It += NumBytesRead;
-		check(It <= End);
+		checkSlow(It <= End);
 		return Out;
 	}
 
diff --git a/Engine/Plugins/Experimental/PlainProps/Source/Public/PlainPropsUeCoreBindings.h b/Engine/Plugins/Experimental/PlainProps/Source/Public/PlainPropsUeCoreBindings.h
--- a/Engine/Plugins/Experimental/PlainProps/Source/Public/PlainPropsUeCoreBindings.h
+++ b/Engine/Plugins/Experimental/PlainProps/Source/Public/PlainPropsUeCoreBindings.h
@@ -362,34 +362,20 @@ struct FTransformBinding : public ICustomBinding
 {
 	using Type = FTransform;
 	inline static constexpr EMemberPresence Occupancy = EMemberPresence::AllowSparse;
-	//inline static constexpr uint8 NumMembers = 3;
-	//enum class EMember : uint8 { Translate, Rotate, Scale };
-	//FMemberId MemberIds[3];
-	//FStructSchemaId VectorId;
-	//FStructSchemaId QuatId;
-
-	enum class EMember : uint8 { TranslateX, TranslateY, TranslateZ, RotateX, RotateY, RotateZ, RotateW, ScaleX, ScaleY, ScaleZ };
-	FMemberId MemberIds[11];
+	enum class EMember : uint8 { Translate, Rotate, Scale };
+	FMemberId MemberIds[3];
+	FStructSchemaId VectorId;
+	FStructSchemaId QuatId;
 
 	template<typename Ids>
-	void InitIds() 
+	void InitIds(/*const FDeclarations& Declared*/) 
 	{
-		//MemberIds[EMember::Translate] = Ids::IndexMember("Translate");
-		//MemberIds[EMember::Rotate] = Ids::IndexMember("Rotate");
-		//MemberIds[EMember::Scale] = Ids::IndexMember("Scale");
-		//VectorId = IndexNativeStruct<FVector, Ids>();
-		//QuatId = IndexNativeStruct<FVector, Ids>();
-		
-		MemberIds[(uint8)EMember::TranslateX] = Ids::IndexMember("TranslateX");
-		MemberIds[(uint8)EMember::TranslateY] = Ids::IndexMember("TranslateY");
-		MemberIds[(uint8)EMember::TranslateZ] = Ids::IndexMember("TranslateZ");
-		MemberIds[(uint8)EMember::RotateX] = Ids::IndexMember("RotateX");
-		MemberIds[(uint8)EMember::RotateY] = Ids::IndexMember("RotateY");
-		MemberIds[(uint8)EMember::RotateZ] = Ids::IndexMember("RotateZ");
-		MemberIds[(uint8)EMember::RotateW] = Ids::IndexMember("RotateW");
-		MemberIds[(uint8)EMember::ScaleX] = Ids::IndexMember("ScaleX");
-		MemberIds[(uint8)EMember::ScaleY] = Ids::IndexMember("ScaleY");
-		MemberIds[(uint8)EMember::ScaleZ] = Ids::IndexMember("ScaleZ");
+		MemberIds[(uint8)EMember::Translate] = Ids::IndexMember("Translate");
+		MemberIds[(uint8)EMember::Rotate] = Ids::IndexMember("Rotate");
+		MemberIds[(uint8)EMember::Scale] = Ids::IndexMember("Scale");
+
+		VectorId = IndexNativeStruct<FVector, Ids>();
+		QuatId = IndexNativeStruct<FQuat, Ids>();
 	}
 
 	PLAINPROPS_API void	Save(FMemberBuilder& Dst, const FTransform& Src, const FTransform* Default, const FSaveContext& Context) const;
@@ -429,9 +415,9 @@ struct TSetDeltaBinding : public ICustomBinding
 	
 	static TConstArrayView<FMemberId> GetMemberIds() { return FSetOps::Get<Ids>().All; }
 
-	virtual void SaveStruct(FMemberBuilder& Dst, const void* Src, const void* Default, const FSaveContext& Ctx) override;
+	virtual void SaveCustom(FMemberBuilder& Dst, const void* Src, const void* Default, const FSaveContext& Ctx) override;
 
-	virtual void LoadStruct(void* Dst, FStructView Src, ECustomLoadMethod Method, const FLoadBatch& Batch) const override
+	virtual void LoadCustom(void* Dst, FStructView Src, ECustomLoadMethod Method, const FLoadBatch& Batch) const override
 	{
 		Type& Out = *static_cast<Type*>(Dst);
 		FMemberReader Members(Src);
@@ -476,7 +462,7 @@ struct TSetDeltaBinding : public ICustomBinding
 		check(!Members.HasMore());
 	}
 
-	virtual bool DiffStruct(const void* StructA, const void* StructB) const override
+	virtual bool DiffCustom(const void* StructA, const void* StructB) const override
 	{
 		const Type& A = *static_cast<const Type*>(StructA);
 		const Type& B = *static_cast<const Type*>(StructB);
diff --git a/Engine/Plugins/Experimental/PlainProps/Source/Tests/PlainPropsUeCoreTests.cpp b/Engine/Plugins/Experimental/PlainProps/Source/Tests/PlainPropsUeCoreTests.cpp
--- a/Engine/Plugins/Experimental/PlainProps/Source/Tests/PlainPropsUeCoreTests.cpp
+++ b/Engine/Plugins/Experimental/PlainProps/Source/Tests/PlainPropsUeCoreTests.cpp
@@ -132,19 +132,19 @@ struct FTestCustomBinding : public ICustomBinding
 
 struct FNameBinding : public FTestCustomBinding
 {
-	virtual void SaveStruct(FMemberBuilder& Dst, const void* Src, const void*, const FSaveContext& Ctx) override
+	virtual void SaveCustom(FMemberBuilder& Dst, const void* Src, const void*, const FSaveContext& Ctx) override
 	{
 		FSetElementId Idx = Names.Add(*static_cast<const FName*>(Src));
 		Dst.Add(Declaration.Idx, Idx.AsInteger());
 	}
 
-	virtual void LoadStruct(void* Dst, FStructView Src, ECustomLoadMethod, const FLoadBatch&) const override
+	virtual void LoadCustom(void* Dst, FStructView Src, ECustomLoadMethod, const FLoadBatch&) const override
 	{
 		FSetElementId Idx = FSetElementId::FromInteger(FMemberReader(Src).GrabLeaf().AsS32());
 		*static_cast<FName*>(Dst) = Names.Get(Idx);
 	}
 
-	virtual bool DiffStruct(const void* StructA, const void* StructB) const override
+	virtual bool DiffCustom(const void* StructA, const void* StructB) const override
 	{
 		FName A = *static_cast<const FName*>(StructA);
 		FName B = *static_cast<const FName*>(StructB);
@@ -754,6 +754,8 @@ TEST_CASE_NAMED(FPlainPropsUeCoreTest, "System::Core::Serialization::PlainProps:
 
 	SECTION("Transform")
 	{
+		TScopedStructBinding<FVector, EMemberPresence::RequireAll> Vector;
+		TScopedStructBinding<FQuat, EMemberPresence::RequireAll> Quat;
 		BindCustomStructOnce<FTransformBinding, FDefaultRuntime>();
 
 		Run([](FBatchSaver& Batch)
