PlainProps noting optimizations, 74 -> 52ms for point cloud test

#rnx

diff --git a/Engine/Plugins/Experimental/PlainProps/Source/Private/PlainPropsBuildSchema.cpp b/Engine/Plugins/Experimental/PlainProps/Source/Private/PlainPropsBuildSchema.cpp
--- a/Engine/Plugins/Experimental/PlainProps/Source/Private/PlainPropsBuildSchema.cpp
+++ b/Engine/Plugins/Experimental/PlainProps/Source/Private/PlainPropsBuildSchema.cpp
@@ -46,12 +46,14 @@ static FString PrintMemberSchema(FMemberSchema Schema)
 
 struct FStructSchemaBuilder
 {
+	FStructSchemaBuilder(const FStructDeclaration& InDecl, FSchemasBuilder& InSchemas);
+
 	const FStructDeclaration&					Declaration;
 	FSchemasBuilder&							AllSchemas;
-	FScratchAllocator&							Scratch;
-	const FDebugIds&							Debug;
-	TMap<FOptionalMemberId, FMemberSchema>		NotedMembers;
 	bool										bMissingMemberNoted = false;
+	FOptionalMemberId*							MemberOrder;
+	FMemberSchema*								NotedSchemas;
+	TBitArray<>									NotedMembers;
 
 	void										NoteMembersRecursively(const FBuiltStruct& Struct);
 	void										NoteRangeRecursively(ERangeSizeType NumType, TConstArrayView<FMemberType> Types, void* InnermostSchemaBuilder, const FBuiltRange& Range);
@@ -72,7 +74,9 @@ struct FEnumSchemaBuilder
 
 FSchemasBuilder::FSchemasBuilder(const FDeclarations& Types, FScratchAllocator& InScratch)
 : FSchemasBuilder(Types.GetStructs(), Types.GetEnums(), Types.GetDebug(), InScratch)
-{}
+{
+
+}
 
 FSchemasBuilder::FSchemasBuilder(FStructDeclarations InStructs, FEnumDeclarations InEnums, const FDebugIds& InDebug, FScratchAllocator& InScratch)
 : DeclaredStructs(InStructs)
@@ -87,7 +91,7 @@ FSchemasBuilder::FSchemasBuilder(FStructDeclarations InStructs, FEnumDeclaration
 FSchemasBuilder::~FSchemasBuilder() {}
 
 template<class T, typename ...Ts>
-T& GetOrEmplace(int32& Index, TPagedArray<T, 4096>& Things, Ts&&... EmplaceArgs)
+FORCEINLINE T& GetOrEmplace(int32& Index, TPagedArray<T, 4096>& Things, Ts&&... EmplaceArgs)
 {
 	if (Index == INDEX_NONE)
 	{
@@ -98,18 +102,18 @@ T& GetOrEmplace(int32& Index, TPagedArray<T, 4096>& Things, Ts&&... EmplaceArgs)
 	return Things[Index];
 }
 	
-FEnumSchemaBuilder&	FSchemasBuilder::NoteEnum(FEnumSchemaId Id)
+FORCEINLINE FEnumSchemaBuilder&	FSchemasBuilder::NoteEnum(FEnumSchemaId Id)
 {
 	checkf(!bBuilt, TEXT("Noted new members after building"));
 	checkf(DeclaredEnums[Id.Idx], TEXT("Undeclared enum '%s' noted"), *Debug.Print(Id));
 	return GetOrEmplace(EnumIndices[Id.Idx], Enums, *DeclaredEnums[Id.Idx], Id);
 }
 
-FStructSchemaBuilder& FSchemasBuilder::NoteStruct(FStructSchemaId Id)
+FORCEINLINE FStructSchemaBuilder& FSchemasBuilder::NoteStruct(FStructSchemaId Id)
 {
 	checkf(!bBuilt, TEXT("Noted new members after building"));
 	checkf(DeclaredStructs[Id.Idx], TEXT("Undeclared struct '%s' noted"), *Debug.Print(Id));
-	return GetOrEmplace(StructIndices[Id.Idx], Structs, *DeclaredStructs[Id.Idx], *this, Scratch, Debug);
+	return GetOrEmplace(StructIndices[Id.Idx], Structs, *DeclaredStructs[Id.Idx], *this);
 }
 
 void FSchemasBuilder::NoteStructAndMembers(FStructSchemaId Id, const FBuiltStruct& Struct)
@@ -147,13 +151,27 @@ void FSchemasBuilder::NoteInheritanceChains()
 		for (FOptionalStructSchemaId Super = Structs[Idx].Declaration.Super; Super; Super = DeclaredStructs[Super.Get().Idx]->Super)
 		{
 			uint32 SuperIdx = Super.Get().Idx;
-			GetOrEmplace(StructIndices[SuperIdx], Structs, *DeclaredStructs[SuperIdx], *this, Scratch, Debug);
+			GetOrEmplace(StructIndices[SuperIdx], Structs, *DeclaredStructs[SuperIdx], *this);
 		}
 	}
 }
 
 //////////////////////////////////////////////////////////////////////////
 
+FStructSchemaBuilder::FStructSchemaBuilder(const FStructDeclaration& Decl, FSchemasBuilder& Schemas)
+: Declaration(Decl)
+, AllSchemas(Schemas)
+, NotedMembers(false, Declaration.NumMembers + int32(!!Declaration.Super))
+{
+	MemberOrder = Schemas.GetScratch().AllocateArray<FOptionalMemberId>(NotedMembers.Num());
+	NotedSchemas = Schemas.GetScratch().AllocateArray<FMemberSchema>(NotedMembers.Num());
+
+	check(NotedMembers.Num());
+	MemberOrder[0] = NoId;
+	TConstArrayView<FMemberId> Order = Declaration.GetMemberOrder();
+	FMemory::Memcpy(MemberOrder + int32(!!Declaration.Super), Order.GetData(), Order.NumBytes());
+}
+
 static bool RequiresDynamicStructSchema(const FMemberSchema& A, const FMemberSchema& B)
 {
 	if (A.InnerSchema != B.InnerSchema && A.Type.GetKind() == B.Type.GetKind())
@@ -188,56 +206,74 @@ static void* NoteStructOrEnum(FSchemasBuilder& AllSchemas, bool bStruct, FSchema
 void FStructSchemaBuilder::NoteMembersRecursively(const FBuiltStruct& Struct)
 {
 	check(Declaration.Occupancy != EMemberPresence::RequireAll || Struct.NumMembers == Declaration.NumMembers);
-	bMissingMemberNoted |= Struct.NumMembers < static_cast<uint32>(Declaration.NumMembers + !!Declaration.Super);
+	bMissingMemberNoted |= Struct.NumMembers < NotedMembers.Num();
 	
+	if (Struct.NumMembers == 0)
+	{
+		return;
+	}
+
+
+	const int32 NumNoted = NotedMembers.Num();
+	int32 NoteIdx = 0;
 	for (const FBuiltMember& Member : MakeArrayView(Struct.Members, Struct.NumMembers))
 	{
-		if (FMemberSchema* Schema = NotedMembers.Find(Member.Name))
+		while (MemberOrder[NoteIdx] != Member.Name)
+		{
+			++NoteIdx;
+			check(NoteIdx < NumNoted);
+		}
+
+		if (NotedMembers[NoteIdx])
 		{
-			if (RequiresDynamicStructSchema(*Schema, Member.Schema))
+			FMemberSchema& NotedSchema = NotedSchemas[NoteIdx];
+			if (RequiresDynamicStructSchema(NotedSchema, Member.Schema))
 			{
-				if (!Schema->GetInnermostType().AsStruct().IsDynamic)
+				if (!NotedSchema.GetInnermostType().AsStruct().IsDynamic)
 				{
-					SetIsDynamic(Schema->EditInnermostType(Scratch));
-					Schema->InnerSchema = NoId;
+					SetIsDynamic(NotedSchema.EditInnermostType(AllSchemas.GetScratch()));
+					NotedSchema.InnerSchema = NoId;
 				}
-				check(Schema->InnerSchema == NoId);
+				check(NotedSchema.InnerSchema == NoId);
 				
 			}
 			else
 			{
-				checkf(*Schema == Member.Schema, TEXT("Member '%s' in '%s' first added as %s and later as %s."),
-					*Debug.Print(Member.Name), *Debug.Print(Declaration.Type), *PrintMemberSchema(*Schema), *PrintMemberSchema(Member.Schema));
+				checkf(NotedSchema == Member.Schema, TEXT("Member '%s' in '%s' first added as %s and later as %s."),
+					*AllSchemas.GetDebug().Print(Member.Name), *AllSchemas.GetDebug().Print(Declaration.Type), *PrintMemberSchema(NotedSchema), *PrintMemberSchema(Member.Schema));
 			}
 		}
 		else
 		{
-			NotedMembers.Add(Member.Name, Member.Schema);
+			NotedMembers[NoteIdx] = true;
+			NotedSchemas[NoteIdx] = Member.Schema;
 		}
+
+		++NoteIdx;
 		
 		const FMemberSchema& Schema = Member.Schema;
-		if (EMemberKind Kind = Schema.Type.GetKind(); Kind == EMemberKind::Leaf)
-		{
-			if (ELeafType::Enum == Schema.Type.AsLeaf().Type)
-			{
-				AllSchemas.NoteEnum(static_cast<FEnumSchemaId>(Schema.InnerSchema.Get())).NoteValue(Member.Value.Leaf);
-			}
-		}
-		else if (Kind == EMemberKind::Struct)
-		{
-			AllSchemas.NoteStruct(static_cast<FStructSchemaId>(Schema.InnerSchema.Get())).NoteMembersRecursively(*Member.Value.Struct);
-		}
-		else 
+		if (Schema.InnerSchema)
 		{
-			check(IsStructOrEnum(Schema.GetInnermostType()) == !!Schema.InnerSchema);
+			checkSlow(IsStructOrEnum(Schema.GetInnermostType()));
+			FSchemaId InnerSchema = Schema.InnerSchema.Get();
 
-			if (Schema.InnerSchema)
+			switch (Schema.Type.GetKind())
 			{
-				void* InnerSchemaBuilder = NoteStructOrEnum(AllSchemas, Schema.GetInnermostType().IsStruct(), Schema.InnerSchema.Get());
+			case EMemberKind::Leaf:
+				AllSchemas.NoteEnum(static_cast<FEnumSchemaId>(InnerSchema)).NoteValue(Member.Value.Leaf);
+				break;
+
+			case EMemberKind::Struct:
+				AllSchemas.NoteStruct(static_cast<FStructSchemaId>(InnerSchema)).NoteMembersRecursively(*Member.Value.Struct);
+				break;
+		
+			case EMemberKind::Range:
+				void* InnerSchemaBuilder = NoteStructOrEnum(AllSchemas, Schema.GetInnermostType().IsStruct(), InnerSchema);
 				if (Member.Value.Range)
 				{
 					NoteRangeRecursively(Schema.Type.AsRange().MaxSize, Schema.GetInnerRangeTypes(), InnerSchemaBuilder, *Member.Value.Range);			
 				}
+				break;
 			}
 		}
 	}
@@ -295,35 +331,25 @@ FBuiltStructSchema FStructSchemaBuilder::Build() const
 	FBuiltStructSchema Out = { Declaration.Type, Declaration.Id, Declaration.Super };
 	Out.bDense = Declaration.Occupancy == EMemberPresence::RequireAll || !bMissingMemberNoted;
 	
-	if (int32 Num = NotedMembers.Num())
+	if (int32 Num = NotedMembers.CountSetBits())
 	{
 		Out.MemberNames.Reserve(Num);
 		Out.MemberSchemas.Reserve(Num);
-
-		// Add generated super struct before declared members
-		if (Declaration.Super)
+		for (int32 NoteIdx = 0, NoteNum = NotedMembers.Num(); NoteIdx < NoteNum; ++NoteIdx)
 		{
-			for (const TPair<FOptionalMemberId, FMemberSchema>& NotedMember : NotedMembers)
+			if (NotedMembers[NoteIdx])
 			{
-				if (NotedMember.Key == NoId)
+				if (FOptionalMemberId Name = MemberOrder[NoteIdx])
 				{
-					Out.MemberSchemas.Add(&NotedMember.Value);
-					break;
+					Out.MemberNames.Add(Name.Get());
 				}
+				Out.MemberSchemas.Add(&NotedSchemas[NoteIdx]);
 			}
 		}
 
-		for (FMemberId Name : Declaration.GetMemberOrder())
-		{
-			if (const FMemberSchema* Schema = NotedMembers.Find(ToOptional(Name)))
-			{
-				Out.MemberNames.Add(Name);
-				Out.MemberSchemas.Add(Schema);
-			}
-		}
+		check(Num == Out.MemberSchemas.Num());
 	}
 
-	check(NotedMembers.Num() == Out.MemberSchemas.Num());
 	return Out;
 }
 
diff --git a/Engine/Plugins/Experimental/PlainProps/Source/Public/PlainPropsBuild.h b/Engine/Plugins/Experimental/PlainProps/Source/Public/PlainPropsBuild.h
--- a/Engine/Plugins/Experimental/PlainProps/Source/Public/PlainPropsBuild.h
+++ b/Engine/Plugins/Experimental/PlainProps/Source/Public/PlainPropsBuild.h
@@ -108,6 +108,11 @@ struct FMemberSchema
 
 inline bool operator==(FMemberSchema A, FMemberSchema B)
 {
+	if (FMemory::Memcmp(&A, &B, sizeof(FMemberSchema)) == 0)
+	{
+		return true;
+	}
+
 	return A.Type == B.Type && A.InnerSchema == B.InnerSchema && Algo::Compare(A.GetInnerRangeTypes(), B.GetInnerRangeTypes());
 }
 //////////////////////////////////////////////////////////////////////////
@@ -315,7 +320,7 @@ public:
 
 	void Add(FTypedRange Range)
 	{
-		check(Range.Schema == Schema || Range.Values == nullptr);
+		check(Range.Values == nullptr || Range.Schema == Schema);
 		Ranges.Add(Range.Values);
 	}
 
diff --git a/Engine/Plugins/Experimental/PlainProps/Source/Public/PlainPropsBuildSchema.h b/Engine/Plugins/Experimental/PlainProps/Source/Public/PlainPropsBuildSchema.h
--- a/Engine/Plugins/Experimental/PlainProps/Source/Public/PlainPropsBuildSchema.h
+++ b/Engine/Plugins/Experimental/PlainProps/Source/Public/PlainPropsBuildSchema.h
@@ -54,6 +54,9 @@ public:
 	PLAINPROPS_API FStructSchemaBuilder&		NoteStruct(FStructSchemaId Id);
 	PLAINPROPS_API void							NoteStructAndMembers(FStructSchemaId Id, const FBuiltStruct& Struct);
 	PLAINPROPS_API FBuiltSchemas				Build();
+	
+	FScratchAllocator&							GetScratch() const { return Scratch; }
+	const FDebugIds&							GetDebug() const { return Debug; }
 
 private:
 	FStructDeclarations							DeclaredStructs;
