PlainProps finish TSet range bindings

#rnx

diff --git a/Engine/Plugins/Experimental/PlainProps/Source/Private/PlainPropsSave.cpp b/Engine/Plugins/Experimental/PlainProps/Source/Private/PlainPropsSave.cpp
--- a/Engine/Plugins/Experimental/PlainProps/Source/Private/PlainPropsSave.cpp
+++ b/Engine/Plugins/Experimental/PlainProps/Source/Private/PlainPropsSave.cpp
@@ -63,7 +63,6 @@ struct FLeafRangeSaver
 
 	void Append(FExistingItemSlice Slice, uint32 Stride, SIZE_T LeafSize, const FSaveContext&)
 	{
-		check(Stride == LeafSize);
 		check(OutIt + Slice.Num * LeafSize <= OutEnd);
 		FMemory::Memcpy(OutIt, Slice.Data, Slice.Num * LeafSize);
 		OutIt += Slice.Num * LeafSize;
@@ -76,15 +75,32 @@ struct FLeafRangeSaver
 	}
 };
 
+template<SIZE_T LeafSize>
+struct TStridingLeafRangeSaver : FLeafRangeSaver
+{
+	inline TStridingLeafRangeSaver(FScratchAllocator& Scratch, uint64 Num, SIZE_T) : FLeafRangeSaver(Scratch, Num, LeafSize) {}
+
+	inline void Append(FExistingItemSlice Slice, uint32 Stride, SIZE_T, const FSaveContext&)
+	{
+		uint8* Dst = OutIt;
+		for (const uint8* Src = static_cast<const uint8*>(Slice.Data), *SrcEnd = Src + Slice.Num * Stride; Src != SrcEnd; Src += Stride, Dst += LeafSize)
+		{
+			FMemory::Memcpy(Dst, Src, LeafSize);
+		}
+		OutIt = Dst;
+	}
+
+};
+
 //////////////////////////////////////////////////////////////////////////
 
 template<typename BuiltItemType, typename ItemSchemaType>
-struct TStructuralRangeSaver
+struct TNonLeafRangeSaver
 {
 	FBuiltRange* Out;
 	BuiltItemType* It;
 
-	TStructuralRangeSaver(FScratchAllocator& Scratch, uint64 Num, ItemSchemaType)
+	TNonLeafRangeSaver(FScratchAllocator& Scratch, uint64 Num, ItemSchemaType)
 	: Out(FBuiltRange::Create(Scratch, Num, sizeof(BuiltItemType)))
 	, It(reinterpret_cast<BuiltItemType*>(Out->Data))
 	{}
@@ -105,37 +121,69 @@ struct TStructuralRangeSaver
 	}
 };
 
-using FNestedRangeSaver = TStructuralRangeSaver<FBuiltRange*, FRangeMemberBinding>;
-using FStructRangeSaver = TStructuralRangeSaver<FBuiltStructPtr, FStructSchemaId>;
-using FStructRangeDeltaSaver = TStructuralRangeSaver<FBuiltStructPtr, FDefaultStruct>;
+using FNestedRangeSaver = TNonLeafRangeSaver<FBuiltRange*, FRangeMemberBinding>;
+using FStructRangeSaver = TNonLeafRangeSaver<FBuiltStruct*, FStructSchemaId>;
+using FStructRangeDeltaSaver = TNonLeafRangeSaver<FBuiltStruct*, FDefaultStruct>;
 
 //////////////////////////////////////////////////////////////////////////
 
+
 template<class SaverType, typename InnerContextType>
-[[nodiscard]] FBuiltRange* SaveRange(const void* Range, const IItemRangeBinding& Binding, const FSaveContext& OuterCtx, InnerContextType InnerCtx)
+[[nodiscard]] inline FBuiltRange* SaveRangeItems(FSaveRangeContext& ReadCtx, const IItemRangeBinding& Binding, const FSaveContext& OuterCtx, InnerContextType InnerCtx)
+{
+	const uint64 NumTotal = ReadCtx.Items.NumTotal;
+	SaverType Saver(OuterCtx.Scratch, NumTotal, InnerCtx);
+	while (true)
+	{
+		check(ReadCtx.Items.Slice.Num > 0);
+		Saver.Append(ReadCtx.Items.Slice, ReadCtx.Items.Stride, InnerCtx, OuterCtx);
+	
+		ReadCtx.Request.NumRead += ReadCtx.Items.Slice.Num;
+		if (ReadCtx.Request.NumRead >= NumTotal)
+		{
+			check(ReadCtx.Request.NumRead == NumTotal);	
+			return Saver.Finish();
+		}
+
+		Binding.ReadItems(ReadCtx);	
+	}
+}
+
+template<class SaverType, typename InnerContextType>
+[[nodiscard]] FBuiltRange* SaveNonLeafRange(const void* Range, const IItemRangeBinding& Binding, const FSaveContext& OuterCtx, InnerContextType InnerCtx)
+{
+	FSaveRangeContext ReadCtx = { { Range } };
+	Binding.ReadItems(ReadCtx);
+
+	if (ReadCtx.Items.NumTotal)
+	{
+		return SaveRangeItems<SaverType>(ReadCtx, Binding, OuterCtx, InnerCtx);
+	}
+	
+	return nullptr;
+}
+
+[[nodiscard]] static FBuiltRange* SaveLeafRange(const void* Range, const IItemRangeBinding& Binding, const FSaveContext& OuterCtx, ELeafWidth Width)
 {
+	SIZE_T LeafSize = SizeOf(Width);
 	FSaveRangeContext ReadCtx = { { Range } };
 	Binding.ReadItems(ReadCtx);
 
 	if (const uint64 NumTotal = ReadCtx.Items.NumTotal)
 	{
-		SaverType Saver(OuterCtx.Scratch, NumTotal, InnerCtx);
-		while (true)
+		if (ReadCtx.Items.Stride == LeafSize)
 		{
-			check(ReadCtx.Items.Slice.Num > 0);
-			Saver.Append(ReadCtx.Items.Slice, ReadCtx.Items.Stride, InnerCtx, OuterCtx);
-		
-			ReadCtx.Request.NumRead += ReadCtx.Items.Slice.Num;
-			if (ReadCtx.Request.NumRead >= NumTotal)
-			{
-				check(ReadCtx.Request.NumRead == NumTotal);	
-				return Saver.Finish();
-			}
-
-			Binding.ReadItems(ReadCtx);	
+			return SaveRangeItems<FLeafRangeSaver>(ReadCtx, Binding, OuterCtx, LeafSize);
+		}
+		else switch (Width)
+		{
+			case ELeafWidth::B8:	return SaveRangeItems<TStridingLeafRangeSaver<1>>(ReadCtx, Binding, OuterCtx, LeafSize);
+			case ELeafWidth::B16:	return SaveRangeItems<TStridingLeafRangeSaver<2>>(ReadCtx, Binding, OuterCtx, LeafSize);
+			case ELeafWidth::B32:	return SaveRangeItems<TStridingLeafRangeSaver<4>>(ReadCtx, Binding, OuterCtx, LeafSize);
+			case ELeafWidth::B64:	return SaveRangeItems<TStridingLeafRangeSaver<8>>(ReadCtx, Binding, OuterCtx, LeafSize);
 		}
 	}
-	
+
 	return nullptr;
 }
 
@@ -146,13 +194,7 @@ template<class SaverType, typename InnerContextType>
 	return { Member.InnerTypes + 1, Member.RangeBindings + 1, static_cast<uint16>(Member.NumRanges - 1), Member.InnermostSchema };
 }
 
- ELeafWidth GetArithmeticWidth(FLeafBindType Leaf)
- {
-	checkf(Leaf.Bind.Type != ELeafBindType::BitfieldBool, TEXT("Arrays of bitfields is not a thing"));
-	return Leaf.Arithmetic.Width;
- }
-
-[[nodiscard]] static FBuiltRange* SaveLeafRange(FScratchAllocator& Scratch, const uint8* Range, const ILeafRangeBinding& Binding, FUnpackedLeafType Leaf)
+[[nodiscard]] static FBuiltRange* SaveLeafRangeBinding(FScratchAllocator& Scratch, const uint8* Range, const ILeafRangeBinding& Binding, FUnpackedLeafType Leaf)
 {
 	FLeafRangeAllocator Allocator(Scratch, Leaf);
 	Binding.SaveLeaves(Range, Allocator);
@@ -163,14 +205,20 @@ template<class SaverType, typename InnerContextType>
 {
 	if (const FDefaultStruct* Defaults = Algo::FindBy(Ctx.Defaults, Id, &FDefaultStruct::Id))
 	{
-		return SaveRange<FStructRangeDeltaSaver>(Range, ItemBinding, Ctx, *Defaults);
+		return SaveNonLeafRange<FStructRangeDeltaSaver>(Range, ItemBinding, Ctx, *Defaults);
 	}
 	else
 	{
-		return SaveRange<FStructRangeSaver>(Range, ItemBinding, Ctx, Id);
+		return SaveNonLeafRange<FStructRangeSaver>(Range, ItemBinding, Ctx, Id);
 	}
 }
 
+ELeafWidth GetArithmeticWidth(FLeafBindType Leaf)
+{
+	checkf(Leaf.Bind.Type != ELeafBindType::BitfieldBool, TEXT("Arrays of bitfields is not a thing"));
+	return Leaf.Arithmetic.Width;
+}
+
 [[nodiscard]] static FBuiltRange* SaveRange(const uint8* Range, FRangeMemberBinding Member, const FSaveContext& Ctx)
 {
 	FRangeBinding Binding = Member.RangeBindings[0];
@@ -178,14 +226,14 @@ template<class SaverType, typename InnerContextType>
 
 	if (Binding.IsLeafBinding())
 	{
-		return SaveLeafRange(Ctx.Scratch, Range, Binding.AsLeafBinding(), UnpackNonBitfield(InnerType.AsLeaf()));
+		return SaveLeafRangeBinding(Ctx.Scratch, Range, Binding.AsLeafBinding(), UnpackNonBitfield(InnerType.AsLeaf()));
 	}
 
 	const IItemRangeBinding& ItemBinding = Binding.AsItemBinding();
 	switch (InnerType.GetKind())
 	{
-	case EMemberKind::Leaf:		return SaveRange<FLeafRangeSaver>(Range, ItemBinding, Ctx, SizeOf(GetArithmeticWidth(InnerType.AsLeaf())));
-	case EMemberKind::Range:	return SaveRange<FNestedRangeSaver>(Range, ItemBinding, Ctx, GetInnerRange(Member));
+	case EMemberKind::Leaf:		return SaveLeafRange(Range, ItemBinding, Ctx, GetArithmeticWidth(InnerType.AsLeaf()));
+	case EMemberKind::Range:	return SaveNonLeafRange<FNestedRangeSaver>(Range, ItemBinding, Ctx, GetInnerRange(Member));
 	case EMemberKind::Struct:	return SaveStructRange(Range, ItemBinding, Ctx, static_cast<FStructSchemaId>(Member.InnermostSchema.Get()));
 	}
 
@@ -200,12 +248,12 @@ template<class SaverType, typename InnerContextType>
 	return SaveRange(Range, Member, Ctx);
 }
 
-[[nodiscard]] static FBuiltStructPtr SaveRangeItem(const uint8* Struct, FStructSchemaId Id, const FSaveContext& Ctx)
+[[nodiscard]] static FBuiltStruct* SaveRangeItem(const uint8* Struct, FStructSchemaId Id, const FSaveContext& Ctx)
 {
 	return SaveStruct(Struct, Id, Ctx);
 }
 
-[[nodiscard]] static FBuiltStructPtr SaveRangeItem(const uint8* Struct, FDefaultStruct Default, const FSaveContext& Ctx)
+[[nodiscard]] static FBuiltStruct* SaveRangeItem(const uint8* Struct, FDefaultStruct Default, const FSaveContext& Ctx)
 {
 	return SaveStructDelta(Struct, Default.Struct, Default.Id, Ctx);
 }
@@ -262,7 +310,7 @@ static void SaveMember(FMemberBuilder& Out, const void* Struct, FMemberId Name,
 	Out.AddStruct(Name, Member.Id, SaveStruct(At(Struct, Member.Offset), Member.Id, Ctx));
 }
 
-FBuiltStructPtr SaveStruct(const void* Struct, FStructSchemaId Id, const FSaveContext& Ctx)
+FBuiltStruct* SaveStruct(const void* Struct, FStructSchemaId Id, const FSaveContext& Ctx)
 {
 	const FStructDeclaration& Declaration = Ctx.Declarations.Get(Id);
 
@@ -470,13 +518,13 @@ static void SaveMemberDelta(FMemberBuilder& Out, const void* Struct, const void*
 
 static void SaveMemberDelta(FMemberBuilder& Out, const void* Struct, const void* Default, FMemberId Name, const FSaveContext& Ctx, FStructMemberBinding Member)
 {
-	if (FBuiltStructPtr Delta = SaveStructDelta(At(Struct, Member.Offset), At(Default, Member.Offset), Member.Id, Ctx))
+	if (FBuiltStruct* Delta = SaveStructDelta(At(Struct, Member.Offset), At(Default, Member.Offset), Member.Id, Ctx))
 	{
 		Out.AddStruct(Name, Member.Id, MoveTemp(Delta));
 	}
 }
 
-FBuiltStructPtr SaveStructDelta(const void* Struct, const void* Default, FStructSchemaId Id, const FSaveContext& Ctx)
+FBuiltStruct* SaveStructDelta(const void* Struct, const void* Default, FStructSchemaId Id, const FSaveContext& Ctx)
 {
 	const FStructDeclaration& Declaration = Ctx.Declarations.Get(Id);
 
diff --git a/Engine/Plugins/Experimental/PlainProps/Source/Public/PlainPropsBind.h b/Engine/Plugins/Experimental/PlainProps/Source/Public/PlainPropsBind.h
--- a/Engine/Plugins/Experimental/PlainProps/Source/Public/PlainPropsBind.h
+++ b/Engine/Plugins/Experimental/PlainProps/Source/Public/PlainPropsBind.h
@@ -362,10 +362,12 @@ struct FExistingItemSlice
 	const void*		Data = nullptr;			
 	uint64			Num = 0;
 
+	explicit operator bool() const { return !!Num; }
+
 	const uint8* At(uint64 Idx, uint32 Stride) const
 	{
 		check(Idx < Num);
-		return reinterpret_cast<const uint8*>(Data) + Idx*Stride;
+		return reinterpret_cast<const uint8*>(Data) + Idx * Stride;
 	}
 };
 
@@ -375,12 +377,6 @@ struct FExistingItems
 	uint32				Stride = 0;
 	FExistingItemSlice	Slice;
 
-	void SetPart(FExistingItemSlice Part)
-	{
-		NumTotal += Part.Num;
-		Slice = Part;
-	}
-
 	void SetAll(FExistingItemSlice Whole, uint32 InStride)
 	{
 		NumTotal = Whole.Num;
diff --git a/Engine/Plugins/Experimental/PlainProps/Source/Public/PlainPropsUeCoreBindings.h b/Engine/Plugins/Experimental/PlainProps/Source/Public/PlainPropsUeCoreBindings.h
--- a/Engine/Plugins/Experimental/PlainProps/Source/Public/PlainPropsUeCoreBindings.h
+++ b/Engine/Plugins/Experimental/PlainProps/Source/Public/PlainPropsUeCoreBindings.h
@@ -260,11 +260,11 @@ struct TSetBinding : public IItemRangeBinding
 
 	virtual void MakeItems(FLoadRangeContext& Ctx) const override
 	{
-		const TSet<T>& Set = *Ctx.Request.GetRange<TSet<T>>();
+		TSet<T>& Set = Ctx.Request.GetRange<TSet<T>>();
 		SizeType Num = static_cast<SizeType>(Ctx.Request.NumTotal());
 
 		static constexpr bool bAllocate = sizeof(T) > sizeof(FLoadRangeContext::Scratch);
-		static constexpr uint64 MaxItems = bAllocate ? sizeof(FLoadRangeContext::Scratch) / sizeof(T) : 1;
+		static constexpr uint64 MaxItems = bAllocate ? 1 : sizeof(FLoadRangeContext::Scratch) / SIZE_T(sizeof(T));
 		
 		if (Ctx.Request.IsFirstCall())
 		{
@@ -280,7 +280,7 @@ struct TSetBinding : public IItemRangeBinding
 				Ctx.Items.Set(Tmp, NumTmp, sizeof(T));
 				if constexpr (std::is_default_constructible_v<T>)
 				{
-					for (T* It = static_cast<T*>(Tmp), End = It + NumTmp; It != End; ++It)
+					for (T* It = static_cast<T*>(Tmp), *End = It + NumTmp; It != End; ++It)
 					{
 						::new (It) T;
 					}
@@ -317,7 +317,7 @@ struct TSetBinding : public IItemRangeBinding
 			}
 			else
 			{
-				Ctx.Items.Set(Tmp.Slice(0, FMath::Min(static_cast<uint64>(Tmp.Num()), Ctx.Request.NumMore())));
+				Ctx.Items.Set(Tmp.GetData(), FMath::Min(static_cast<uint64>(Tmp.Num()), Ctx.Request.NumMore()));
 				check(Ctx.Items.Get<T>().Num());
 			}
 		}
@@ -326,32 +326,34 @@ struct TSetBinding : public IItemRangeBinding
 	virtual void ReadItems(FSaveRangeContext& Ctx) const override
 	{
 		static_assert(offsetof(TSetElement<T>, Value) == 0);
-		const TSparseArray<TSetElement<T>>& Elems = *Ctx.Request.GetRange<TSparseArray<TSetElement<T>>>();
+		const TSparseArray<TSetElement<T>>& Elems = Ctx.Request.GetRange<TSparseArray<TSetElement<T>>>();
 
-		if (int32 NumRead = static_cast<int32>(Ctx.Items.NumTotal))
+		if (FExistingItemSlice LastRead = Ctx.Items.Slice)
 		{
 			// Continue partial response
-			const TSetElement<T>* NextElem = static_cast<const TSetElement<T>*>(Ctx.Items.Slice) + Ctx.Items.Slice.Num + /* skip known invalid */ 1;
-			Ctx.Items.SetPart(GetContiguousSlice(Elems.PointerToIndex(NextElem), Elems));
+			const TSetElement<T>* NextElem = static_cast<const TSetElement<T>*>(LastRead.Data) + LastRead.Num + /* skip known invalid */ 1;
+			Ctx.Items.Slice = GetContiguousSlice(Elems.PointerToIndex(NextElem), Elems);
 		}
 		else if (Elems.IsCompact())
 		{
-			Ctx.Items.SetAll(Elems.GetData(), Elems.Num());
+			int32 Num = Elems.Num();
+			Ctx.Items.SetAll(Num ? &Elems[0] : nullptr, Num);
 		}
 		else
 		{
 			// Start partial response
+			Ctx.Items.NumTotal = Elems.Num();
 			Ctx.Items.Stride = sizeof(TSetElement<T>);
-			Ctx.Items.SetPart(GetContiguousSlice(0, Elems));
+			Ctx.Items.Slice = GetContiguousSlice(0, Elems);
 		}
 	}
 
-	FExistingItemSlice GetContiguousSlice(int32 Idx, const TSparseArray<TSetElement<T>>& Elems)
+	static FExistingItemSlice GetContiguousSlice(int32 Idx, const TSparseArray<TSetElement<T>>& Elems)
 	{
 		int32 Num = 1;
 		for (;!Elems.IsValidIndex(Idx); ++Idx) {}
 		for (; Elems.IsValidIndex(Idx + Num); ++Num) {}
-		return { &Elems[Idx], Num };
+		return { &Elems[Idx], static_cast<uint64>(Num) };
 	}
 };
 
diff --git a/Engine/Plugins/Experimental/PlainProps/Source/Tests/PlainPropsUeCoreTests.cpp b/Engine/Plugins/Experimental/PlainProps/Source/Tests/PlainPropsUeCoreTests.cpp
--- a/Engine/Plugins/Experimental/PlainProps/Source/Tests/PlainPropsUeCoreTests.cpp
+++ b/Engine/Plugins/Experimental/PlainProps/Source/Tests/PlainPropsUeCoreTests.cpp
@@ -190,14 +190,14 @@ FBatchSaver::FBatchSaver()
 template<class T>
 void FBatchSaver::Save(T&& Object) 
 {
-	FStructSchemaId Id = IndexNativeStruct<T, FIds>();
+	FStructSchemaId Id = IndexNativeStruct<std::remove_reference_t<T>, FIds>();
 	SavedObjects.Emplace(Id, SaveStruct(&Object, Id, {GTypes, GSchemas, Customs, Scratch}));
 }
 
 template<class T>
 bool FBatchSaver::SaveDelta(const T& Object, const T& Default) 
 {
-	FStructSchemaId Id = IndexNativeStruct<T, FIds>();
+	FStructSchemaId Id = IndexNativeStruct<std::remove_reference_t<T>, FIds>();
 	if (FBuiltStructPtr Delta = SaveStructDelta(&Object, &Default, Id, {GTypes, GSchemas, Customs, Scratch}))
 	{
 		SavedObjects.Emplace(Id, MoveTemp(Delta));
@@ -407,10 +407,20 @@ static void Run(void (*Save)(FBatchSaver&), void (*Load)(FBatchLoader&))
 //{
 //	TArray<FSub> Subs;
 //};
+//
+//struct FNDC
+//{
+//	int X = -1;
+//	explicit FNDC(int I) : X(I) {}
+//	friend bool operator==(FNDC A, FNDC B) = default;
+//	friend inline uint32 GetTypeHash(FNDC I) { return ::GetTypeHash(I.X); }
+//};
+//PP_REFLECT_STRUCT(PlainProps::UE::Test, FNDC, void, X);
 
 struct FInt { int32 X; };
 PP_REFLECT_STRUCT(PlainProps::UE::Test, FInt, void, X);
 static bool operator==(FInt A, FInt B) { return A.X == B.X; }
+inline uint32 GetTypeHash(FInt I) { return ::GetTypeHash(I.X); }
 
 enum class EFlat1 : uint8 { A = 1, B = 3 };
 enum class EFlat2 : uint8 { A, B };
@@ -427,17 +437,19 @@ struct FEnums
 	EFlat2 Flat2;
 	EFlag1 Flag1;
 	EFlag2 Flag2;
+
+	friend bool operator==(FEnums, FEnums) = default;
 };
 PP_REFLECT_STRUCT(PlainProps::UE::Test, FEnums, void, Flat1, Flat2, Flag1, Flag2);
-static bool operator==(FEnums A, FEnums B) { return A.Flat1 == B.Flat1 && A.Flat2 == B.Flat2 && A.Flag1 == B.Flag1 && A.Flag2 == B.Flag2; }
 
 struct FLeafArrays
 {
 	TArray<bool> Bits;
 	TArray<int>	Bobs;
+
+	bool operator==(const FLeafArrays&) const = default;
 };
 PP_REFLECT_STRUCT(PlainProps::UE::Test, FLeafArrays, void, Bits, Bobs);
-static bool operator==(const FLeafArrays& A, const FLeafArrays& B) { return A.Bits == B.Bits && A.Bobs == B.Bobs; }
 
 struct FComplexArrays
 {
@@ -445,17 +457,20 @@ struct FComplexArrays
 	TArray<EFlat1> Enums;
 	TArray<FLeafArrays> Misc;
 	TArray<TArray<EFlat1>> Nested;
+
+	bool operator==(const FComplexArrays&) const = default;
 };
 PP_REFLECT_STRUCT(PlainProps::UE::Test, FComplexArrays, void, Str, Enums, Misc, Nested);
-static bool operator==(const FComplexArrays& A, const FComplexArrays& B) { return A.Str == B.Str && A.Enums == B.Enums && A.Misc == B.Misc && A.Nested == B.Nested; }
 
 struct FNames
 {
 	FName Name;
 	TArray<FName> Names;
+
+	bool operator==(const FNames&) const = default;
 };
 PP_REFLECT_STRUCT(PlainProps::UE::Test, FNames, void, Name, Names);
-static bool operator==(const FNames& A, const FNames& B) { return A.Name == B.Name && A.Names == B.Names; }
+
 
 struct FStr
 {
@@ -463,6 +478,18 @@ struct FStr
 };
 PP_REFLECT_STRUCT(PlainProps::UE::Test, FStr, void, S);
 
+struct FSets
+{
+	TSet<char> Leaves;
+	TSet<TArray<uint8>> Ranges;
+	TSet<FInt> Structs;
+};
+PP_REFLECT_STRUCT(PlainProps::UE::Test, FSets, void, Leaves, Ranges, Structs);
+
+inline bool operator==(const FSets& A, const FSets& B)
+{
+	return LegacyCompareEqual(A.Leaves, B.Leaves) && LegacyCompareEqual(A.Ranges, B.Ranges) && LegacyCompareEqual(A.Structs, B.Structs);
+}
 
 //////////////////////////////////////////////////////////////////////////
 
@@ -513,9 +540,10 @@ struct FDelta
 	FInt		C = { 2 };
 	TArray<int> D;
 	FString		E = "!";
+
+	friend bool operator==(const FDelta& A, const FDelta& B) = default;
 };
 PP_REFLECT_STRUCT(PlainProps::UE::Test, FDelta, void, A, B, C, D, E);
-static bool operator==(const FDelta& A, const FDelta& B) { return A.A == B.A && A.B == B.B && A.C == B.C && A.D == B.D && A.E == B.E; }
 
 //////////////////////////////////////////////////////////////////////////
 //
@@ -718,7 +746,33 @@ TEST_CASE_NAMED(FPlainPropsUeCoreTest, "System::Core::Serialization::PlainProps:
 	}
 		
 	SECTION("TSet")
-	{}
+	{
+		TScopedStructBinding<FInt> Int;
+		TScopedStructBinding<FSets> Sets;
+		Run([](FBatchSaver& Batch)
+			{
+				Batch.Save(FSets{{'H','i'}, {{uint8(10)}, {}}, {{123}}});
+				
+				// Test order preservation
+				Batch.Save(FSets{{'a','b'}});
+				Batch.Save(FSets{{'b','a'}});
+
+				// Test non-compact set
+				FSets Sparse = FSets{{'w','z','a','p','?','!'}};
+				Sparse.Leaves.Remove('w');
+				Sparse.Leaves.Remove('p');
+				Sparse.Leaves.Remove('!');
+				Batch.Save(Sparse);
+			}, 
+			[](FBatchLoader& Batch)
+			{
+				CHECK(Batch.Load<FSets>() == FSets{{'H','i'}, {{uint8(10)}, {}}, {{123}}});
+				CHECK(FSets{{'a','b'}} != FSets{{'b','a'}});
+				CHECK(Batch.Load<FSets>() == FSets{{'a','b'}});
+				CHECK(Batch.Load<FSets>() == FSets{{'b','a'}});
+				CHECK(Batch.Load<FSets>() == FSets{{'z','a','?'}});
+			});
+	}
 
 	SECTION("Delta")
 	{
