PlainProps add TOptional bindings and document custom binding use cases

#rnx

diff --git a/Engine/Plugins/Experimental/PlainProps/Source/Public/PlainPropsBind.h b/Engine/Plugins/Experimental/PlainProps/Source/Public/PlainPropsBind.h
--- a/Engine/Plugins/Experimental/PlainProps/Source/Public/PlainPropsBind.h
+++ b/Engine/Plugins/Experimental/PlainProps/Source/Public/PlainPropsBind.h
@@ -235,6 +235,13 @@ protected: // for unit tests
 
 enum ECustomLoadMethod { Construct, Assign };
 
+/// Load/save a struct with custom code to handle:
+/// * reference types
+/// * private members
+/// * non-default constructible types
+/// * custom delta semantics
+/// * other runtime representations than struct/class, e.g. serialize database
+/// * optimization for very common struct
 class ICustomBinding
 {
 public:
diff --git a/Engine/Plugins/Experimental/PlainProps/Source/Public/PlainPropsUeCoreBindings.h b/Engine/Plugins/Experimental/PlainProps/Source/Public/PlainPropsUeCoreBindings.h
--- a/Engine/Plugins/Experimental/PlainProps/Source/Public/PlainPropsUeCoreBindings.h
+++ b/Engine/Plugins/Experimental/PlainProps/Source/Public/PlainPropsUeCoreBindings.h
@@ -6,6 +6,7 @@
 #include "Math/Quat.h"
 #include "Math/Vector.h"
 #include "Math/Transform.h"
+#include "Misc/Optional.h"
 #include "PlainPropsBind.h"
 #include "PlainPropsLoad.h"
 #include "PlainPropsRead.h"
@@ -23,7 +24,6 @@ PP_REFLECT_STRUCT(, FQuat, void, X, Y, Z, W);
 PP_NAME_STRUCT(, FTransform);
 }
 
-
 namespace PlainProps::UE
 {
 
@@ -252,6 +252,58 @@ struct TUniquePtrBinding : public IItemRangeBinding
 
 //////////////////////////////////////////////////////////////////////////
 
+template <typename T>
+struct TOptionalBinding : public IItemRangeBinding
+{
+	using SizeType = bool;
+	using ItemType = T;
+
+	virtual void MakeItems(FLoadRangeContext& Ctx) const override
+	{
+		TOptional<T>& Opt = Ctx.Request.GetRange<TOptional<T>>();
+		
+		if (Ctx.Request.NumTotal() == 0)
+		{
+			Opt.Reset();
+		}
+		else if constexpr (std::is_default_constructible_v<T>)
+		{
+			if (!Opt)
+			{
+				Opt.Emplace();
+			}
+			Ctx.Items.Set(reinterpret_cast<T*>(&Opt), 1);
+		}
+		else if (Opt)
+		{
+			Ctx.Items.Set(reinterpret_cast<T*>(&Opt), 1);
+		}
+		else
+		{
+			if (Ctx.Request.IsFirstCall())
+			{
+				Ctx.Items.SetUnconstructed();
+				Ctx.Items.RequestFinalCall();
+				Ctx.Items.Set(reinterpret_cast<T*>(&Opt), 1);	
+			}
+			else
+			{
+				// Move-construct from self reference
+				Opt.Emplace(reinterpret_cast<T&&>(Opt));
+			}
+		}
+	}
+
+	virtual void ReadItems(FSaveRangeContext& Ctx) const override
+	{
+		const TOptional<T>& Opt = Ctx.Request.GetRange<TOptional<T>>();
+		check(!Opt || reinterpret_cast<const T*>(&Opt) == &Opt.GetValue());
+		Ctx.Items.SetAll(reinterpret_cast<const T*>(Opt ? &Opt : nullptr), Opt ? 1 : 0);
+	}
+};
+
+//////////////////////////////////////////////////////////////////////////
+
 template <typename T>
 struct TSetBinding : public IItemRangeBinding
 {
@@ -601,8 +653,6 @@ struct TSetDeltaBinding : public ICustomBinding
 }
 
 
-
-
 namespace PlainProps
 {
 
@@ -633,6 +683,13 @@ struct TRangeBind<TSet<T>>
 	using Type = UE::TSetBinding<T>;
 };
 
+template<typename T>
+struct TRangeBind<TOptional<T>>
+{
+	using Type = UE::TOptionalBinding<T>;
+};
+
+
 template<>
 struct TCustomBind<FTransform>
 {
diff --git a/Engine/Plugins/Experimental/PlainProps/Source/Tests/PlainPropsUeCoreTests.cpp b/Engine/Plugins/Experimental/PlainProps/Source/Tests/PlainPropsUeCoreTests.cpp
--- a/Engine/Plugins/Experimental/PlainProps/Source/Tests/PlainPropsUeCoreTests.cpp
+++ b/Engine/Plugins/Experimental/PlainProps/Source/Tests/PlainPropsUeCoreTests.cpp
@@ -407,15 +407,6 @@ static void Run(void (*Save)(FBatchSaver&), void (*Load)(FBatchLoader&))
 //{
 //	TArray<FSub> Subs;
 //};
-//
-//struct FNDC
-//{
-//	int X = -1;
-//	explicit FNDC(int I) : X(I) {}
-//	friend bool operator==(FNDC A, FNDC B) = default;
-//	friend inline uint32 GetTypeHash(FNDC I) { return ::GetTypeHash(I.X); }
-//};
-//PP_REFLECT_STRUCT(PlainProps::UE::Test, FNDC, void, X);
 
 struct FInt { int32 X; };
 PP_REFLECT_STRUCT(PlainProps::UE::Test, FInt, void, X);
@@ -533,6 +524,36 @@ TArray<TUniquePtr<T>> MakeTwo(T&& A, T&& B)
 
 //////////////////////////////////////////////////////////////////////////
 
+struct FNDC
+{
+	int X = -1;
+	explicit FNDC(int I) : X(I) {}
+	friend bool operator==(FNDC A, FNDC B) = default;
+};
+PP_REFLECT_STRUCT(PlainProps::UE::Test, FNDC, void, X);
+
+struct FNDCIntrusive : FNDC
+{
+	FNDCIntrusive() : FNDC(-1) {}
+	explicit FNDCIntrusive(FIntrusiveUnsetOptionalState) : FNDC(-1) {}
+	explicit FNDCIntrusive(int I) : FNDC(I) {}
+	friend bool operator==(FNDCIntrusive, FNDCIntrusive) = default;
+	bool operator==(FIntrusiveUnsetOptionalState) const { return X == -1;}
+};
+PP_REFLECT_STRUCT(PlainProps::UE::Test, FNDCIntrusive, void, X);
+
+struct FOpts
+{
+	TOptional<bool> Bit;
+	TOptional<FNDC> NDC;
+	TOptional<FNDCIntrusive> NDCI;
+
+	friend bool operator==(const FOpts&, const FOpts&) = default;
+};
+PP_REFLECT_STRUCT(PlainProps::UE::Test, FOpts, void, Bit, NDC, NDCI);
+
+//////////////////////////////////////////////////////////////////////////
+
 struct FDelta
 {
 	bool		A = true;
@@ -703,6 +724,27 @@ TEST_CASE_NAMED(FPlainPropsUeCoreTest, "System::Core::Serialization::PlainProps:
 			});
 	}
 
+	SECTION("TOptional")
+	{
+		TScopedStructBinding<FNDC> NDC;
+		TScopedStructBinding<FNDCIntrusive> NDCI;
+		TScopedStructBinding<FOpts> Opts;
+		Run([](FBatchSaver& Batch)
+			{
+				Batch.Save(FOpts{});
+				Batch.Save(FOpts{{true}, {FNDC{2}}, {FNDCIntrusive{3}}});
+				Batch.Save(FOpts{{true}, {FNDC{2}}, {FNDCIntrusive{3}}});
+			}, 
+			[](FBatchLoader& Batch)
+			{
+				CHECK(Batch.Load<FOpts>() == FOpts{});
+				CHECK(Batch.Load<FOpts>() == FOpts{{true}, {FNDC{2}}, {FNDCIntrusive{3}}});
+				FOpts AlreadySet = {{false}, {FNDC{0}}, {FNDCIntrusive{1}}};
+				Batch.LoadInto(AlreadySet);
+				CHECK(AlreadySet == FOpts{{true}, {FNDC{2}}, {FNDCIntrusive{3}}});
+			});
+	}
+
 	SECTION("FName")
 	{
 		TScopedStructBinding<FNames> Names;
