PlainProps

+ TMap range binding
+ Anonymous parametric type names used for range bound template parameters
+ Unscoped type names used by arithmetic anonymous arithmetic and anonymous parametric types
* Fixed nested scope and parametric type bind translation bug
* Improved various error messages

#rnx

diff --git a/Engine/Plugins/Experimental/PlainProps/Source/Private/PlainPropsBind.cpp b/Engine/Plugins/Experimental/PlainProps/Source/Private/PlainPropsBind.cpp
--- a/Engine/Plugins/Experimental/PlainProps/Source/Private/PlainPropsBind.cpp
+++ b/Engine/Plugins/Experimental/PlainProps/Source/Private/PlainPropsBind.cpp
@@ -291,7 +291,7 @@ static void TranslateScopeIds(TArrayView<FNestedScopeId> Out, FIdIndexerBase& In
 		check(Scope.Outer.IsFlat() || Scope.Outer.AsNested().Idx < OutIdx);
 		FScopeId Outer = Scope.Outer.IsFlat() ? FScopeId(Translate(Scope.Outer.AsFlat(), ToNames)) : FScopeId(Out[Scope.Outer.AsNested().Idx]);
 		FFlatScopeId Inner = Translate(Scope.Inner, ToNames);
-		Out[OutIdx] = Indexer.NestScope(Outer, Inner).AsNested();
+		Out[OutIdx++] = Indexer.NestScope(Outer, Inner).AsNested();
 	}
 }
 
@@ -306,7 +306,7 @@ static void TranslateParametricTypeIds(TArrayView<FParametricTypeId> Out, FIdInd
 		{
 			Params.Add(To.Remap(FromParameter));
 		}
-		Out[OutIdx] = Indexer.MakeParametricTypeId(To.Remap(Parametric.Name), Params);
+		Out[OutIdx++] = Indexer.MakeParametricTypeId(To.Remap(Parametric.Name), Params);
 	}
 }
 
@@ -316,6 +316,7 @@ static void TranslateSchemaIds(TArrayView<FSchemaId> Out, FIdIndexerBase& Indexe
 	for (const FStructSchema& FromSchema : GetStructSchemas(From))
 	{
 		FTypeId ToType = To.Remap(FromSchema.Type);
+		checkSlow(ToType.Name.NumParameters == FromSchema.Type.Name.NumParameters);
 		Out[OutIdx++] = Indexer.IndexStruct(ToType);
 	}
 	
diff --git a/Engine/Plugins/Experimental/PlainProps/Source/Private/PlainPropsBuildSchema.cpp b/Engine/Plugins/Experimental/PlainProps/Source/Private/PlainPropsBuildSchema.cpp
--- a/Engine/Plugins/Experimental/PlainProps/Source/Private/PlainPropsBuildSchema.cpp
+++ b/Engine/Plugins/Experimental/PlainProps/Source/Private/PlainPropsBuildSchema.cpp
@@ -105,14 +105,14 @@ FORCEINLINE T& GetOrEmplace(int32& Index, TPagedArray<T, 4096>& Things, Ts&&...
 FORCEINLINE FEnumSchemaBuilder&	FSchemasBuilder::NoteEnum(FEnumSchemaId Id)
 {
 	checkf(!bBuilt, TEXT("Noted new members after building"));
-	checkf(DeclaredEnums[Id.Idx], TEXT("Undeclared enum '%s' noted"), *Debug.Print(Id));
+	checkf(Id.Idx < uint32(DeclaredEnums.Num()) && DeclaredEnums[Id.Idx], TEXT("Undeclared enum '%s' noted"), *Debug.Print(Id));
 	return GetOrEmplace(EnumIndices[Id.Idx], Enums, *DeclaredEnums[Id.Idx], Id);
 }
 
 FORCEINLINE FStructSchemaBuilder& FSchemasBuilder::NoteStruct(FStructSchemaId Id)
 {
 	checkf(!bBuilt, TEXT("Noted new members after building"));
-	checkf(DeclaredStructs[Id.Idx], TEXT("Undeclared struct '%s' noted"), *Debug.Print(Id));
+	checkf(Id.Idx < uint32(DeclaredStructs.Num()) && DeclaredStructs[Id.Idx], TEXT("Undeclared struct '%s' noted"), *Debug.Print(Id));
 	return GetOrEmplace(StructIndices[Id.Idx], Structs, *DeclaredStructs[Id.Idx], *this);
 }
 
diff --git a/Engine/Plugins/Experimental/PlainProps/Source/Private/PlainPropsDeclare.cpp b/Engine/Plugins/Experimental/PlainProps/Source/Private/PlainPropsDeclare.cpp
--- a/Engine/Plugins/Experimental/PlainProps/Source/Private/PlainPropsDeclare.cpp
+++ b/Engine/Plugins/Experimental/PlainProps/Source/Private/PlainPropsDeclare.cpp
@@ -51,7 +51,7 @@ void FDeclarations::DeclareStruct(FStructSchemaId Id, FTypeId Type, TConstArrayV
 	}
 
 	TUniquePtr<FStructDeclaration>& Ptr = DeclaredStructs[Id.Idx];
-	checkf(!Ptr, TEXT("Already declared"));
+	checkf(!Ptr, TEXT("'%s' is already declared"), *Debug.Print(Id));
 
 	FStructDeclaration Header{Id, Type, Super, Occupancy, IntCastChecked<uint16>(MemberOrder.Num())};
 	void* Data = FMemory::Malloc(sizeof(FStructDeclaration) + MemberOrder.Num() * MemberOrder.GetTypeSize());
@@ -67,7 +67,7 @@ void FDeclarations::DeclareEnum(FEnumSchemaId Id, FTypeId Type, EEnumMode Mode,
 	}
 	
 	TUniquePtr<FEnumDeclaration>& Ptr = DeclaredEnums[Id.Idx];
-	checkf(!Ptr, TEXT("Already declared")); 
+	checkf(!Ptr, TEXT("'%s' is already declared"), *Debug.Print(Id));
 
 	FEnumDeclaration Header{Type, Mode, Width, IntCastChecked<uint16>(Enumerators.Num())};
 	void* Data = FMemory::Malloc(sizeof(FEnumDeclaration) + Enumerators.Num() * Enumerators.GetTypeSize());
diff --git a/Engine/Plugins/Experimental/PlainProps/Source/Private/PlainPropsIndex.cpp b/Engine/Plugins/Experimental/PlainProps/Source/Private/PlainPropsIndex.cpp
--- a/Engine/Plugins/Experimental/PlainProps/Source/Private/PlainPropsIndex.cpp
+++ b/Engine/Plugins/Experimental/PlainProps/Source/Private/PlainPropsIndex.cpp
@@ -31,10 +31,10 @@ FParametricTypeView::FParametricTypeView(FConcreteTypenameId InName, TConstArray
 namespace ParametricTypeHash
 {
 
-static uint32 Calculate(FConcreteTypenameId Name, TConstArrayView<FTypeId> Parameters)
+static uint32 Calculate(FOptionalConcreteTypenameId Name, TConstArrayView<FTypeId> Parameters)
 {
 	uint64 ParametersHash = FXxHash64::HashBuffer(Parameters.GetData(), sizeof(FTypeId) * Parameters.Num()).Hash;
-	return HashCombineFast(Name.Id.Idx, static_cast<uint32>(ParametersHash));
+	return HashCombineFast(GetTypeHash(Name), static_cast<uint32>(ParametersHash));
 }
 
 static constexpr uint8 FreeSlotByte = 0xFF;
@@ -127,7 +127,7 @@ FScopeId FIdIndexerBase::NestScope(FScopeId Outer, FFlatScopeId Inner)
 	return FScopeId(NestedScopes.Index(Outer, Inner));
 }
 
-FParametricTypeId FIdIndexerBase::MakeParametricTypeId(FConcreteTypenameId Name, TConstArrayView<FTypeId> Params)
+FParametricTypeId FIdIndexerBase::MakeParametricTypeId(FOptionalConcreteTypenameId Name, TConstArrayView<FTypeId> Params)
 {
 	return ParametricTypes.Index({Name, IntCastChecked<uint8>(Params.Num()), Params.GetData()});
 }
@@ -137,6 +137,11 @@ FTypeId FIdIndexerBase::MakeParametricType(FTypeId Type, TConstArrayView<FTypeId
 	return {Type.Scope, FTypenameId(MakeParametricTypeId(Type.Name.AsConcrete(), Params))};
 }
 
+FTypeId FIdIndexerBase::MakeAnonymousParametricType(TConstArrayView<FTypeId> Params)
+{
+	return { NoId, FTypenameId(MakeParametricTypeId(NoId, Params)) };
+}
+
 FEnumSchemaId FIdIndexerBase::IndexEnum(FTypeId Type)
 {
 	return { IntCastChecked<uint32>(Enums.Add(Type).AsInteger()) };
diff --git a/Engine/Plugins/Experimental/PlainProps/Source/Private/PlainPropsLoad.cpp b/Engine/Plugins/Experimental/PlainProps/Source/Private/PlainPropsLoad.cpp
--- a/Engine/Plugins/Experimental/PlainProps/Source/Private/PlainPropsLoad.cpp
+++ b/Engine/Plugins/Experimental/PlainProps/Source/Private/PlainPropsLoad.cpp
@@ -249,6 +249,14 @@ static void CloneBindingWithReplacedStructIds(const FSchemaId* FromIds, const FS
 	return Memcpy ? FLoadStructPlan(Memcpy.GetValue()) : MakeSchemaLoadPlan(From, To, ToMemberIds, ToStructIds, OutSubsetSchemas);
 }
 
+class FMemberlessDummyBinding : public ICustomBinding
+{
+	virtual void SaveCustom(FMemberBuilder& Dst, const void* Src, const void* Default, const FSaveContext& Ctx) override { check(false); }
+	virtual void LoadCustom(void* Dst, FStructView Src, ECustomLoadMethod Method, const FLoadBatch& Batch) const override { check(false);}
+	virtual bool DiffCustom(const void* StructA, const void* StructB) const override { check(false); return false; }
+};
+static FMemberlessDummyBinding GMemberlessBinding;
+
 FLoadBatchPtr CreateLoadPlans(FReadBatchId ReadId, const FDeclarations& Declarations, const FCustomBindings& Customs, const FSchemaBindings& Schemas, TConstArrayView<FStructSchemaId> RuntimeIds)
 {
 	check(NumStructSchemas(ReadId) == RuntimeIds.Num());
@@ -268,15 +276,24 @@ FLoadBatchPtr CreateLoadPlans(FReadBatchId ReadId, const FDeclarations& Declarat
 		int32 SubsetSchemaOffset = SubsetSchemaData.Num();
 		if (const ICustomBinding* Custom = Customs.FindStruct(RuntimeId))
 		{
-			Plans[SavedId.Idx] = FLoadStructPlan(*Custom) ;
+			Plans[SavedId.Idx] = FLoadStructPlan(*Custom);
 		}
 		else
 		{
 			const FStructSchema& From = ResolveStructSchema(ReadId, SavedId);
-			const FSchemaBinding& To = Schemas.GetStruct(RuntimeId);
-			// Possible optimization - some simple memcpy cases doesn't need to resolve the declaration
-			TConstArrayView<FMemberId> ToMemberIds = Declarations.Get(RuntimeId).GetMemberOrder();
-			Plans[SavedId.Idx] = MakeLoadPlan(From, To, ToMemberIds, RuntimeIds, /* out */ SubsetSchemaData);	
+			if (From.NumMembers)
+			{
+				const FSchemaBinding& To = Schemas.GetStruct(RuntimeId);
+				// Possible optimization - some simple memcpy cases doesn't need to resolve the declaration
+				TConstArrayView<FMemberId> ToMemberIds = Declarations.Get(RuntimeId).GetMemberOrder();
+				Plans[SavedId.Idx] = MakeLoadPlan(From, To, ToMemberIds, RuntimeIds, /* out */ SubsetSchemaData);	
+			}
+			else
+			{
+				checkf(!From.Type.Scope && From.Type.Name.NumParameters == 2, TEXT("Only range-bound template parameters are memberless. "
+					"They're always anonymous two-parameter types and uninstantiable as structs, bound via MakeAnonymousParametricType()"));
+				Plans[SavedId.Idx] = FLoadStructPlan(GMemberlessBinding);
+			}
 		}
 		
 		SubsetSchemaSizes[SavedId.Idx] = SubsetSchemaData.Num() - SubsetSchemaOffset;
diff --git a/Engine/Plugins/Experimental/PlainProps/Source/Private/PlainPropsWrite.cpp b/Engine/Plugins/Experimental/PlainProps/Source/Private/PlainPropsWrite.cpp
--- a/Engine/Plugins/Experimental/PlainProps/Source/Private/PlainPropsWrite.cpp
+++ b/Engine/Plugins/Experimental/PlainProps/Source/Private/PlainPropsWrite.cpp
@@ -38,13 +38,16 @@ struct FWriteIds
 	FMemberId							Remap(FMemberId Old) const				{ return { Remap(Old.Id) }; }
 	FFlatScopeId						Remap(FFlatScopeId Old) const			{ return { Remap(Old.Name) }; }
 	FNestedScopeId						Remap(FNestedScopeId Old) const			{ return NestedScopes[Old.Idx].Get(); }
-	FScopeId							Remap(FScopeId Old) const				{ return Old.IsNested() ? FScopeId(Remap(Old.AsNested())) : FScopeId(Remap(Old.AsFlat())); }
+	FScopeId							Remap(FScopeId Old) const				{ return Old.IsFlat() ? FScopeId(Remap(Old.AsFlat())) : Old ? FScopeId(Remap(Old.AsNested())) : Old; }
 	FConcreteTypenameId					Remap(FConcreteTypenameId Old) const	{ return { Remap(Old.Id) }; }	
 	FParametricTypeId					Remap(FParametricTypeId Old) const		{ return ParametricTypes[Old.Idx].Get(); }
 	FTypenameId							Remap(FTypenameId Old) const			{ return Old.IsConcrete() ? FTypenameId(Remap(Old.AsConcrete())) : FTypenameId(Remap(Old.AsParametric())); }
 	FTypeId								Remap(FTypeId Old) const				{ return { Remap(Old.Scope), Remap(Old.Name) }; }
 	FStructSchemaId						RemapStruct(FSchemaId Old) const		{ return Structs[Old.Idx].Get(); }
 	FEnumSchemaId						RemapEnum(FSchemaId Old) const			{ return Enums[Old.Idx].Get(); }
+
+	template<typename T>
+	TOptionalId<T>						Remap(TOptionalId<T> Old) const			{ return Old ? ToOptional(Remap(Old.Get())) : Old; }
 };
 
 static TConstArrayView<FNameId> GetUsedNames(const FBuiltStructSchema& Used)
@@ -92,6 +95,14 @@ struct FUsedIds
 		Names[Name.Idx] = true;
 	}
 	
+	void MarkUsed(FOptionalConcreteTypenameId Name)
+	{
+		if (Name)
+		{
+			MarkUsed(Name.Get().Id);
+		}
+	}
+	
 	void MarkUsed(FTypeId Type)
 	{
 		MarkUsed(Type.Scope);
@@ -104,7 +115,7 @@ struct FUsedIds
 		{
 			MarkUsed(Scope.AsFlat().Name);
 		}
-		else
+		else if (Scope)
 		{
 			FBitReference Used = NestedScopes[Scope.AsNested().Idx];
 			if (!Used)
@@ -132,7 +143,7 @@ struct FUsedIds
 				Used = true;
 
 				FParametricTypeView ParametricType = Declared.Resolve(Typename.AsParametric());
-				MarkUsed(ParametricType.Name.Id);
+				MarkUsed(ParametricType.Name);
 				for (FTypeId Parameter : ParametricType.GetParameters())
 				{
 					MarkUsed(Parameter);
@@ -779,7 +790,7 @@ void FDebugIds::AppendDebugString(FString& Out, FScopeId Scope) const
 	{
 		AppendDebugString(Out, Scope.AsFlat().Name);
 	}
-	else
+	else if (Scope)
 	{
 		FNestedScope Nested = Resolve(Scope.AsNested());
 		AppendDebugString(Out, Nested.Outer);
@@ -797,21 +808,29 @@ void FDebugIds::AppendDebugString(FString& Out, FTypenameId Typename) const
 	else
 	{
 		FParametricTypeView ParametricType = Resolve(Typename.AsParametric());
-		AppendDebugString(Out, ParametricType.Name.Id);
-		Out.AppendChar('<');
+	
+		if (ParametricType.Name)
+		{
+			AppendDebugString(Out, ParametricType.Name.Get().Id);
+		}
+
+		Out.AppendChar(ParametricType.Name ? '<' : '[');
 		for (FTypeId Parameter : ParametricType.GetParameters())
 		{
 			AppendDebugString(Out, Parameter);
 			Out.AppendChar(',');
 		}
-		Out.GetCharArray()[Out.Len() - 1] = '>';
+		Out.GetCharArray()[Out.Len() - 1] = ParametricType.Name ? '>' : ']';
 	}
 }
 
 void FDebugIds::AppendDebugString(FString& Out, FTypeId Type) const
 {
-	AppendDebugString(Out, Type.Scope);
-	Out.AppendChar('.');
+	if (Type.Scope)
+	{
+		AppendDebugString(Out, Type.Scope);
+		Out.AppendChar('.');
+	}
 	AppendDebugString(Out, Type.Name);
 }
 
diff --git a/Engine/Plugins/Experimental/PlainProps/Source/Public/PlainPropsBind.h b/Engine/Plugins/Experimental/PlainProps/Source/Public/PlainPropsBind.h
--- a/Engine/Plugins/Experimental/PlainProps/Source/Public/PlainPropsBind.h
+++ b/Engine/Plugins/Experimental/PlainProps/Source/Public/PlainPropsBind.h
@@ -9,6 +9,7 @@
 #include "PlainPropsDeclare.h"
 #include "PlainPropsRead.h"
 #include "PlainPropsTypes.h"
+#include <tuple>
 
 namespace PlainProps 
 {
@@ -536,37 +537,81 @@ private:
 
 ////////////////////////////////////////////////////////////////////////////////////////////////
 
+template<class Ctti, typename Ids>
+FTypeId IndexStructOrEnumType();
+
 template<typename Enum, typename Ids>
-FEnumSchemaId IndexNativeEnum()
+FEnumSchemaId IndexEnum()
 {
 	static FEnumSchemaId Id = Ids::IndexEnum(CttiOf<Enum>::Name);
 	return Id;
 }
 
 template<typename Struct, typename Ids>
-FStructSchemaId IndexNativeStruct()
+FStructSchemaId IndexStruct()
 {
-	static FStructSchemaId Id = Ids::IndexStruct(CttiOf<Struct>::Name);
+	static FStructSchemaId Id = Ids::IndexStruct(IndexStructOrEnumType<CttiOf<Struct>, Ids>());
 	return Id;
 }
 
 template<typename Struct, typename Ids>
-FOptionalStructSchemaId IndexOptionalNativeStruct()
+FOptionalStructSchemaId IndexOptionalStruct()
 {
 	if constexpr (!std::is_void_v<Struct>)
 	{
-		return IndexNativeStruct<Struct, Ids>();
+		return IndexStruct<Struct, Ids>();
 	}
 	
 	return NoId;
 }
 
+template<typename T, typename Ids>
+FTypeId IndexAnyType()
+{
+	using RangeBinding = RangeBind<T>;
+	if constexpr (!std::is_void_v<RangeBinding>)
+	{
+		FTypeId InnerType = IndexAnyType<typename RangeBinding::ItemType, Ids>();
+		FTypeId SizeType = { NoId, Ids::IndexTypename(CttiOf<typename RangeBinding::SizeType>::Name) };
+		return Ids::GetIndexer().MakeAnonymousParametricType({InnerType, SizeType});
+	}
+	else if constexpr (std::is_arithmetic_v<T>)
+	{
+		return { NoId, Ids::IndexTypename(CttiOf<T>::Name) };
+	}
+	else 
+	{
+		return IndexStructOrEnumType<CttiOf<T>, Ids>();
+	}
+}
+
+template<typename Ids, typename... Ts>
+FTypeId IndexTemplatedType(FTypeId TemplatedType, const std::tuple<Ts...>*)
+{
+	FTypeId Parameters[] = { (IndexAnyType<Ts, Ids>())... };
+	return Ids::GetIndexer().MakeParametricType(TemplatedType, Parameters);
+}
+
+template<class Ctti, typename Ids>
+FTypeId IndexStructOrEnumType()
+{
+	FTypeId Type = Ids::IndexNativeType(Ctti::Name);
+	if constexpr (Templated<Ctti>)
+	{
+		return IndexTemplatedType<Ids>(Type, (typename Ctti::TemplateArgs*)nullptr);
+	}
+	else
+	{
+		return Type;
+	}
+}
+
 ////////////////////////////////////////////////////////////////////////////////////////////////
 
 template<typename Type, class Ids>
 FMemberBindType BindMemberLeaf(FOptionalSchemaId& OutSchema)
 {
-	OutSchema = std::is_enum_v<Type> ? ToOptional(static_cast<FSchemaId>(IndexNativeEnum<Type, Ids>())) : NoId;
+	OutSchema = std::is_enum_v<Type> ? ToOptional(static_cast<FSchemaId>(IndexEnum<Type, Ids>())) : NoId;
 	return FMemberBindType(ReflectLeaf<Type>);
 }
 
@@ -620,7 +665,7 @@ FMemberBindType BindMemberStruct(FOptionalSchemaId& OutSchema)
 {
 	if constexpr (std::is_void_v<CustomBinding>)
 	{
-		OutSchema = FOptionalSchemaId(IndexNativeStruct<Type, typename Runtime::Ids>());
+		OutSchema = FOptionalSchemaId(IndexStruct<Type, typename Runtime::Ids>());
 	}
 	else
 	{
@@ -635,12 +680,12 @@ FMemberBindType BindType(FOptionalSchemaId& OutSchema)
 {
 	if constexpr (std::is_arithmetic_v<Type> || std::is_enum_v<Type>)
 	{
-		OutSchema = std::is_enum_v<Type> ? ToOptional(static_cast<FSchemaId>(IndexNativeEnum<Type, Ids>())) : NoId;
+		OutSchema = std::is_enum_v<Type> ? ToOptional(static_cast<FSchemaId>(IndexEnum<Type, Ids>())) : NoId;
 		return FMemberBindType(ReflectLeaf<Type>);
 	}
 	else
 	{
-		OutSchema = FOptionalSchemaId(IndexNativeStruct<Type, Ids>());
+		OutSchema = FOptionalSchemaId(IndexStruct<Type, Ids>());
 		return FMemberBindType(FStructType{EMemberKind::Struct, /* IsDynamic */ 0, /* IsSuper */ 0});
 	}
 }
@@ -774,9 +819,9 @@ FEnumSchemaId DeclareNativeEnum(FDeclarations& Out, EEnumMode Mode)
 template<class Ctti, class Ids>
 FStructSchemaId DeclareNativeStruct(FDeclarations& Out, EMemberPresence Occupancy)
 {
-	FTypeId Type = Ids::IndexNativeType(Ctti::Name);
+	FTypeId Type = IndexStructOrEnumType<Ctti, Ids>();
 	FStructSchemaId Id = Ids::IndexStruct(Type);
-	FOptionalStructSchemaId SuperId = IndexOptionalNativeStruct<typename Ctti::Super, Ids>();
+	FOptionalStructSchemaId SuperId = IndexOptionalStruct<typename Ctti::Super, Ids>();
 	FMemberId MemberIds[Ctti::NumVars];
 	ForEachVar<Ctti>([&]<class Var>()
 	{ 
@@ -870,12 +915,15 @@ struct FIdBinding
 	FMemberId							Remap(FMemberId Old) const				{ return { Remap(Old.Id) }; }
 	FFlatScopeId						Remap(FFlatScopeId Old) const			{ return { Remap(Old.Name) }; }
 	FNestedScopeId						Remap(FNestedScopeId Old) const			{ return NestedScopes[Old.Idx]; }
-	FScopeId							Remap(FScopeId Old) const				{ return Old.IsNested() ? FScopeId(Remap(Old.AsNested())) : FScopeId(Remap(Old.AsFlat())); }
+	FScopeId							Remap(FScopeId Old) const				{ return Old.IsFlat() ? FScopeId(Remap(Old.AsFlat())) : Old ? FScopeId(Remap(Old.AsNested())) : Old; }
 	FConcreteTypenameId					Remap(FConcreteTypenameId Old) const	{ return { Remap(Old.Id) }; }	
 	FParametricTypeId					Remap(FParametricTypeId Old) const		{ return ParametricTypes[Old.Idx]; }
 	FTypenameId							Remap(FTypenameId Old) const			{ return Old.IsConcrete() ? FTypenameId(Remap(Old.AsConcrete())) : FTypenameId(Remap(Old.AsParametric())); }
 	FTypeId								Remap(FTypeId Old) const				{ return { Remap(Old.Scope), Remap(Old.Name) }; }
 
+	template<typename T>
+	TOptionalId<T>						Remap(TOptionalId<T> Old) const			{ return Old ? ToOptional(Remap(Old.Get())) : Old; }
+
 	TConstArrayView<FStructSchemaId>	GetStructIds(int32 NumStructs) const
 	{
 		// All saved struct schema ids are lower than enum schema ids
diff --git a/Engine/Plugins/Experimental/PlainProps/Source/Public/PlainPropsCtti.h b/Engine/Plugins/Experimental/PlainProps/Source/Public/PlainPropsCtti.h
--- a/Engine/Plugins/Experimental/PlainProps/Source/Public/PlainPropsCtti.h
+++ b/Engine/Plugins/Experimental/PlainProps/Source/Public/PlainPropsCtti.h
@@ -60,7 +60,8 @@ template<typename T> using CttiOf = typename TCttiOf<T>::Type;
 #define PP_REFLECT_STRUCT(NS, T, S, ...)				_PP_REFLECT_STRUCT(_PP_NUM_ARGS(__VA_ARGS__), NS, T, S, __VA_ARGS__)
 #define PP_REFLECT_STRUCT_TEMPLATE(NS, T, S, ...)		_PP_REFLECT_STRUCT_TEMPLATE(_PP_NUM_ARGS(__VA_ARGS__), NS, T, S, __VA_ARGS__)
 #define PP_NAME_STRUCT(NS, T)							struct T##_Ctti { inline static constexpr char Name[] = #T; using Type = NS :: T; }; T##_Ctti CttiOfPtr(T*);
-
+#define PP_NAME_STRUCT_TEMPLATE(NS, T)					template<class... Ts> struct T##_Ctti { inline static constexpr char Name[] = #T; using Type = NS :: T<Ts...>; using TemplateArgs = std::tuple<Ts...>; }; template<class... Ts> T##_Ctti<Ts...> CttiOfPtr(T<Ts...>*);
+		
 // Alternate set of macros to reflect classes/structs with bitfield bool members, e.g. uint8 bOol : 1;
 
 #define PP_REFLECT_STRUCT_ONLY(N, NS, T, S)				struct T##_Ctti { inline static constexpr char Name[] = #T; using Type = NS :: T; using Super = S; static constexpr int NumVars = N; template<int> struct Var; }; T##_Ctti CttiOfPtr(T*);
@@ -91,5 +92,8 @@ static constexpr void ForEachVar(Fn&& Visitor)
 
 //////////////////////////////////////////////////////////////////////////
 
+template<class Ctti>
+concept Templated = requires { typename Ctti::TemplateArgs; };
+
 } // namespace PlainProps
 
diff --git a/Engine/Plugins/Experimental/PlainProps/Source/Public/PlainPropsIndex.h b/Engine/Plugins/Experimental/PlainProps/Source/Public/PlainPropsIndex.h
--- a/Engine/Plugins/Experimental/PlainProps/Source/Public/PlainPropsIndex.h
+++ b/Engine/Plugins/Experimental/PlainProps/Source/Public/PlainPropsIndex.h
@@ -50,8 +50,9 @@ class FIdIndexerBase : public FDebugIds
 {
 public:
 	PLAINPROPS_API FScopeId				NestScope(FScopeId Outer, FFlatScopeId Inner);
-	PLAINPROPS_API FParametricTypeId	MakeParametricTypeId(FConcreteTypenameId Name, TConstArrayView<FTypeId> Params);
+	PLAINPROPS_API FParametricTypeId	MakeParametricTypeId(FOptionalConcreteTypenameId Name, TConstArrayView<FTypeId> Params);
 	PLAINPROPS_API FTypeId				MakeParametricType(FTypeId Type, TConstArrayView<FTypeId> Params);
+	PLAINPROPS_API FTypeId				MakeAnonymousParametricType(TConstArrayView<FTypeId> Params);
 
 	PLAINPROPS_API FEnumSchemaId		IndexEnum(FTypeId Type);
 	PLAINPROPS_API FStructSchemaId		IndexStruct(FTypeId Type);
diff --git a/Engine/Plugins/Experimental/PlainProps/Source/Public/PlainPropsTypes.h b/Engine/Plugins/Experimental/PlainProps/Source/Public/PlainPropsTypes.h
--- a/Engine/Plugins/Experimental/PlainProps/Source/Public/PlainPropsTypes.h
+++ b/Engine/Plugins/Experimental/PlainProps/Source/Public/PlainPropsTypes.h
@@ -130,19 +130,28 @@ struct FEnumSchemaId : FSchemaId {};
 
 //////////////////////////////////////////////////////////////////////////
 
+struct FNoId {};
+inline constexpr FNoId NoId;
+
+//////////////////////////////////////////////////////////////////////////
+
 struct FNestedScopeId { uint32 Idx; };
 struct FFlatScopeId { FNameId Name; };
 
 class FScopeId
 {
 	static constexpr uint32 NestedBit = 0x80000000u;
+	static constexpr uint32 Unscoped = ~0u;
+
 	uint32 Handle;
 public:
+	FScopeId(FNoId) : Handle(Unscoped) {}
 	explicit FScopeId(FFlatScopeId Flat) : Handle(Flat.Name.Idx) 				{ check(AsFlat().Name == Flat.Name); }
 	explicit FScopeId(FNestedScopeId Nested) : Handle(Nested.Idx | NestedBit) 	{ check(AsNested().Idx == Nested.Idx); }
 
-	bool							IsFlat() const 					{ return  !(Handle & NestedBit); }
-	bool							IsNested() const 				{ return !!(Handle & NestedBit); }
+	explicit						operator bool() const			{ return Handle != ~0u; }
+	bool							IsFlat() const 					{ return !(Handle & NestedBit); }
+	bool							IsNested() const 				{ return !!(*this) & !!(Handle & NestedBit); }
 	FFlatScopeId					AsFlat() const 					{ check(IsFlat()); 		return {FNameId{Handle}}; }
 	FNestedScopeId					AsNested() const				{ check(IsNested()); 	return {Handle & ~NestedBit}; }
 	uint32							AsInt() const					{ return Handle; }
@@ -204,15 +213,12 @@ struct FTypeId
 
 //////////////////////////////////////////////////////////////////////////
 
-struct FNoId {};
-inline constexpr FNoId NoId;
-
-
 inline uint32 ToIdx(FNameId Id) { return Id.Idx; }
 inline uint32 ToIdx(FMemberId Name) { return Name.Id.Idx; }
 inline uint32 ToIdx(FSchemaId Id) { return Id.Idx; }
 inline uint32 ToIdx(FNestedScopeId Id) { return Id.Idx; }
 inline uint32 ToIdx(FParametricTypeId Id) { return Id.AsInt(); }
+inline uint32 ToIdx(FConcreteTypenameId Name) { return Name.Id.Idx; }
 
 template<class IdType>
 IdType FromIdx(uint32 Idx) { return {Idx}; }
@@ -265,6 +271,7 @@ using FOptionalStructSchemaId = TOptionalId<FStructSchemaId>;
 using FOptionalEnumSchemaId = TOptionalId<FEnumSchemaId>;
 using FOptionalNestedScopeId = TOptionalId<FNestedScopeId>;
 using FOptionalParametricTypeId = TOptionalId<FParametricTypeId>;
+using FOptionalConcreteTypenameId = TOptionalId<FConcreteTypenameId>;
 
 template<class IdType>
 inline constexpr TOptionalId<IdType> ToOptional(IdType Id) { return Id; }
@@ -276,7 +283,7 @@ inline constexpr FOptionalSchemaId ToOptionalSchema(FStructSchemaId Id) { return
 // Resolved FNestedScopeId
 struct FNestedScope
 {
-	FScopeId		Outer;
+	FScopeId		Outer; // @invariant !!Outer
 	FFlatScopeId 	Inner;
 
 	bool operator==(FNestedScope O) const { return Outer.AsInt() == O.Outer.AsInt() && Inner.Name == O.Inner.Name; }
@@ -289,8 +296,8 @@ struct FParameterIndexRange : FBaseTypenameId { using FBaseTypenameId::FBaseType
 // Name-resolved FParametricTypeId
 struct FParametricType
 {
-	FConcreteTypenameId		Name;
-	FParameterIndexRange	Parameters;
+	FOptionalConcreteTypenameId	Name;
+	FParameterIndexRange		Parameters;
 
 	bool operator==(FParametricType O) const { return Name == O.Name && Parameters.AsInt() == O.Parameters.AsInt(); }
 };
@@ -299,9 +306,11 @@ struct FParametricType
 struct FParametricTypeView
 {
 	FParametricTypeView(FConcreteTypenameId InName, uint8 NumParams, const FTypeId* Params) : Name(InName), NumParameters(NumParams), Parameters(Params) {}
+	FParametricTypeView(FOptionalConcreteTypenameId InName, uint8 NumParams, const FTypeId* Params) : Name(InName), NumParameters(NumParams), Parameters(Params) {}
 	FParametricTypeView(FConcreteTypenameId InName, TConstArrayView<FTypeId> Params);
+	FParametricTypeView(FOptionalConcreteTypenameId InName, TConstArrayView<FTypeId> Params);
 
-	FConcreteTypenameId			Name;
+	FOptionalConcreteTypenameId	Name;
 	uint8						NumParameters;
 	const FTypeId*				Parameters;
 
diff --git a/Engine/Plugins/Experimental/PlainProps/Source/Public/PlainPropsUeCoreBindings.h b/Engine/Plugins/Experimental/PlainProps/Source/Public/PlainPropsUeCoreBindings.h
--- a/Engine/Plugins/Experimental/PlainProps/Source/Public/PlainPropsUeCoreBindings.h
+++ b/Engine/Plugins/Experimental/PlainProps/Source/Public/PlainPropsUeCoreBindings.h
@@ -15,6 +15,7 @@
 
 
 PP_NAME_STRUCT(, FName);
+PP_REFLECT_STRUCT_TEMPLATE(, TTuple, void, Key, Value) // Todo handle TTuple and higher arities
 
 namespace UE::Math
 {
@@ -304,15 +305,16 @@ struct TOptionalBinding : public IItemRangeBinding
 
 //////////////////////////////////////////////////////////////////////////
 
-template <typename T>
+template <typename T, typename KeyFuncs, typename SetAllocator>
 struct TSetBinding : public IItemRangeBinding
 {
 	using SizeType = int32;
 	using ItemType = T;
+	using SetType = TSet<T, KeyFuncs, SetAllocator>;
 
 	virtual void MakeItems(FLoadRangeContext& Ctx) const override
 	{
-		TSet<T>& Set = Ctx.Request.GetRange<TSet<T>>();
+		SetType& Set = Ctx.Request.GetRange<SetType>();
 		SizeType Num = static_cast<SizeType>(Ctx.Request.NumTotal());
 
 		static constexpr bool bAllocate = sizeof(T) > sizeof(FLoadRangeContext::Scratch);
@@ -411,6 +413,12 @@ struct TSetBinding : public IItemRangeBinding
 
 //////////////////////////////////////////////////////////////////////////
 
+template <typename K, typename V, typename SetAllocator, typename KeyFuncs>
+struct TMapBinding : public TSetBinding<TPair<K, V>, KeyFuncs, SetAllocator>
+{};
+
+//////////////////////////////////////////////////////////////////////////
+
 //TODO: macroify, e.g PP_CUSTOM_BIND(PLAINPROPS_API, FTransform, Transform, Translate, Rotate, Scale)
 struct FTransformBinding : public ICustomBinding
 {
@@ -428,8 +436,8 @@ struct FTransformBinding : public ICustomBinding
 		MemberIds[(uint8)EMember::Rotate] = Ids::IndexMember("Rotate");
 		MemberIds[(uint8)EMember::Scale] = Ids::IndexMember("Scale");
 
-		VectorId = IndexNativeStruct<FVector, Ids>();
-		QuatId = IndexNativeStruct<FQuat, Ids>();
+		VectorId = IndexStruct<FVector, Ids>();
+		QuatId = IndexStruct<FQuat, Ids>();
 	}
 
 	PLAINPROPS_API void	Save(FMemberBuilder& Dst, const FTransform& Src, const FTransform* Default, const FSaveContext& Context) const;
@@ -652,7 +660,6 @@ struct TSetDeltaBinding : public ICustomBinding
 
 }
 
-
 namespace PlainProps
 {
 
@@ -677,10 +684,16 @@ struct TRangeBind<TUniquePtr<T>>
 	using Type = UE::TUniquePtrBinding<T>;
 };
 
-template<typename T>
-struct TRangeBind<TSet<T>>
+template <typename T, typename KeyFuncs, typename SetAllocator>
+struct TRangeBind<TSet<T, KeyFuncs, SetAllocator>>
 {
-	using Type = UE::TSetBinding<T>;
+	using Type = UE::TSetBinding<T, KeyFuncs, SetAllocator>;
+};
+
+template <typename K, typename V, typename SetAllocator, typename KeyFuncs>
+struct TRangeBind<TMap<K, V, SetAllocator, KeyFuncs>>
+{
+	using Type = UE::TMapBinding<K, V, SetAllocator, KeyFuncs>;
 };
 
 template<typename T>
@@ -695,4 +708,5 @@ struct TCustomBind<FTransform>
 {
 	using Type = UE::FTransformBinding;
 };
-}
\ No newline at end of file
+
+} // namespace PlainProps
\ No newline at end of file
diff --git a/Engine/Plugins/Experimental/PlainProps/Source/Tests/PlainPropsTests.cpp b/Engine/Plugins/Experimental/PlainProps/Source/Tests/PlainPropsTests.cpp
--- a/Engine/Plugins/Experimental/PlainProps/Source/Tests/PlainPropsTests.cpp
+++ b/Engine/Plugins/Experimental/PlainProps/Source/Tests/PlainPropsTests.cpp
@@ -1279,7 +1279,7 @@ TEST_CASE_NAMED(FPlainPropsReadWriteTest, "System::Core::Serialization::PlainPro
 			CHECK(Names[FlatUsed.Name] == "FlatUsed");
 
 			FParametricTypeView S4 = ResolveUntranslatedParametricType(Batch, S4T.Name.AsParametric());
-			CHECK(Names[S4.Name.Id] == "S4");
+			CHECK(Names[S4.Name.Get().Id] == "S4");
 			CHECK(S4.NumParameters == 2);
 			
 			FTypeId S2T = S4.Parameters[0];
diff --git a/Engine/Plugins/Experimental/PlainProps/Source/Tests/PlainPropsUeCoreTests.cpp b/Engine/Plugins/Experimental/PlainProps/Source/Tests/PlainPropsUeCoreTests.cpp
--- a/Engine/Plugins/Experimental/PlainProps/Source/Tests/PlainPropsUeCoreTests.cpp
+++ b/Engine/Plugins/Experimental/PlainProps/Source/Tests/PlainPropsUeCoreTests.cpp
@@ -38,6 +38,7 @@ struct FIds
 	static FEnumSchemaId	IndexEnum(FAnsiStringView Name)				{ return IndexEnum(IndexNativeType(Name)); }
 	static FStructSchemaId	IndexStruct(FTypeId Type)					{ return GNames.IndexStruct(Type); }
 	static FStructSchemaId	IndexStruct(FAnsiStringView Name)			{ return IndexStruct(IndexNativeType(Name)); }
+	static FIdIndexerBase&	GetIndexer()								{ return GNames; }
 	static const FDebugIds& GetDebug()									{ return GNames; }
 };
 
@@ -190,14 +191,14 @@ FBatchSaver::FBatchSaver()
 template<class T>
 void FBatchSaver::Save(T&& Object) 
 {
-	FStructSchemaId Id = IndexNativeStruct<std::remove_reference_t<T>, FIds>();
+	FStructSchemaId Id = IndexStruct<std::remove_reference_t<T>, FIds>();
 	SavedObjects.Emplace(Id, SaveStruct(&Object, Id, {GTypes, GSchemas, Customs, Scratch}));
 }
 
 template<class T>
 bool FBatchSaver::SaveDelta(const T& Object, const T& Default) 
 {
-	FStructSchemaId Id = IndexNativeStruct<std::remove_reference_t<T>, FIds>();
+	FStructSchemaId Id = IndexStruct<std::remove_reference_t<T>, FIds>();
 	if (FBuiltStructPtr Delta = SaveStructDelta(&Object, &Default, Id, {GTypes, GSchemas, Customs, Scratch}))
 	{
 		SavedObjects.Emplace(Id, MoveTemp(Delta));
@@ -462,13 +463,20 @@ struct FNames
 };
 PP_REFLECT_STRUCT(PlainProps::UE::Test, FNames, void, Name, Names);
 
-
 struct FStr
 {
 	FString S;
 };
 PP_REFLECT_STRUCT(PlainProps::UE::Test, FStr, void, S);
 
+struct FNDC
+{
+	int X = -1;
+	explicit FNDC(int I) : X(I) {}
+	friend bool operator==(FNDC A, FNDC B) = default;
+};
+PP_REFLECT_STRUCT(PlainProps::UE::Test, FNDC, void, X);
+
 struct FSets
 {
 	TSet<char> Leaves;
@@ -482,6 +490,19 @@ inline bool operator==(const FSets& A, const FSets& B)
 	return LegacyCompareEqual(A.Leaves, B.Leaves) && LegacyCompareEqual(A.Ranges, B.Ranges) && LegacyCompareEqual(A.Structs, B.Structs);
 }
 
+struct FMaps
+{
+	TMap<bool, bool> Leaves;
+	TMap<int, TArray<char>> Ranges;
+	TMap<FInt, FNDC> Structs;
+};
+PP_REFLECT_STRUCT(PlainProps::UE::Test, FMaps, void, Leaves, Ranges, Structs);
+
+inline bool operator==(const FMaps& A, const FMaps& B)
+{
+	return LegacyCompareEqual(A.Leaves, B.Leaves) && LegacyCompareEqual(A.Ranges, B.Ranges) && LegacyCompareEqual(A.Structs, B.Structs);
+}
+
 //////////////////////////////////////////////////////////////////////////
 
 struct FUniquePtrs
@@ -524,14 +545,6 @@ TArray<TUniquePtr<T>> MakeTwo(T&& A, T&& B)
 
 //////////////////////////////////////////////////////////////////////////
 
-struct FNDC
-{
-	int X = -1;
-	explicit FNDC(int I) : X(I) {}
-	friend bool operator==(FNDC A, FNDC B) = default;
-};
-PP_REFLECT_STRUCT(PlainProps::UE::Test, FNDC, void, X);
-
 struct FNDCIntrusive : FNDC
 {
 	FNDCIntrusive() : FNDC(-1) {}
@@ -816,6 +829,27 @@ TEST_CASE_NAMED(FPlainPropsUeCoreTest, "System::Core::Serialization::PlainProps:
 			});
 	}
 
+	SECTION("TMap")
+	{
+		TScopedStructBinding<FInt> Int;
+		TScopedStructBinding<FNDC> NDC;
+		TScopedStructBinding<FMaps> Maps;
+		TScopedStructBinding<TPair<bool, bool>> BoolBoolPair;
+		TScopedStructBinding<TPair<int, TArray<char>>> IntStringPair;
+		TScopedStructBinding<TPair<FInt, FNDC>> IntNDCPair;
+		
+		Run([](FBatchSaver& Batch)
+			{
+				Batch.Save(FMaps{});
+				Batch.Save(FMaps{{{true, true}, {false, false}}, {{5, {'h', 'i'}}}, {{{7}, FNDC{8}}}});
+			}, 
+			[](FBatchLoader& Batch)
+			{
+				CHECK(Batch.Load<FMaps>() == FMaps{});
+				CHECK(Batch.Load<FMaps>() == FMaps{{{true, true}, {false, false}}, {{5, {'h', 'i'}}}, {{{7}, FNDC{8}}}});
+			});
+	}
+
 	SECTION("Delta")
 	{
 		TScopedStructBinding<FInt> Int;
