PlainProps

+ Enumeration, Arithmetic and LeafType concepts
* Replaced ReflectLeaf with ReflectArithmetic and ReflectEnum
+ TArray allocator support, limited value til typename indirection is added

#rnx

diff --git a/Engine/Plugins/Experimental/PlainProps/Source/Private/PlainPropsRead.cpp b/Engine/Plugins/Experimental/PlainProps/Source/Private/PlainPropsRead.cpp
--- a/Engine/Plugins/Experimental/PlainProps/Source/Private/PlainPropsRead.cpp
+++ b/Engine/Plugins/Experimental/PlainProps/Source/Private/PlainPropsRead.cpp
@@ -340,8 +340,8 @@ FOptionalSchemaId FMemberReader::GrabRangeSchema(FMemberType InnermostType)
 		return FOptionalSchemaId(GrabStructSchema(InnermostType.AsStruct()));
 	}
 
-	ELeafType LeafType = InnermostType.AsLeaf().Type;
-	return LeafType == ELeafType::Enum ? FOptionalSchemaId(GrabEnumSchema()) : NoId;
+	ELeafType Type = InnermostType.AsLeaf().Type;
+	return Type == ELeafType::Enum ? FOptionalSchemaId(GrabEnumSchema()) : NoId;
 }
 
 FLeafView FMemberReader::GrabLeaf()
diff --git a/Engine/Plugins/Experimental/PlainProps/Source/Private/PlainPropsWrite.cpp b/Engine/Plugins/Experimental/PlainProps/Source/Private/PlainPropsWrite.cpp
--- a/Engine/Plugins/Experimental/PlainProps/Source/Private/PlainPropsWrite.cpp
+++ b/Engine/Plugins/Experimental/PlainProps/Source/Private/PlainPropsWrite.cpp
@@ -650,16 +650,16 @@ private:
 		}
 	}
 
-	void WriteLeaf(FLeafType LeafType, uint64 LeafValue)
+	void WriteLeaf(FLeafType Leaf, uint64 LeafValue)
 	{
-		if (LeafType.Type == ELeafType::Bool)
+		if (Leaf.Type == ELeafType::Bool)
 		{
 			check(LeafValue <= 1);
 			Bits.WriteBit(!!LeafValue);
 		}
 		else
 		{
-			WriteUnsigned(LeafValue, SizeOf(LeafType.Width));
+			WriteUnsigned(LeafValue, SizeOf(Leaf.Width));
 		}
 	}
 		
diff --git a/Engine/Plugins/Experimental/PlainProps/Source/Public/PlainPropsBind.h b/Engine/Plugins/Experimental/PlainProps/Source/Public/PlainPropsBind.h
--- a/Engine/Plugins/Experimental/PlainProps/Source/Public/PlainPropsBind.h
+++ b/Engine/Plugins/Experimental/PlainProps/Source/Public/PlainPropsBind.h
@@ -433,11 +433,11 @@ class FLeafRangeAllocator
 public:
 	FLeafRangeAllocator(FScratchAllocator& InScratch, FUnpackedLeafType InExpected) : Scratch(InScratch), Expected(InExpected) {}
 
-	template<typename LeafType, typename SizeType>
-	LeafType* AllocateRange(SizeType Num)
+	template<LeafType T, typename SizeType>
+	T* AllocateRange(SizeType Num)
 	{
-		check(ReflectLeaf<LeafType> == Expected);
-		return Num ? static_cast<LeafType*>(Allocate(ReflectLeaf<LeafType>, IntCastChecked<uint64>(Num))) : nullptr;
+		check(ReflectArithmetic<T> == Expected);
+		return Num ? static_cast<T*>(Allocate(ReflectArithmetic<T>, IntCastChecked<uint64>(Num))) : nullptr;
 	}
 
 	FBuiltRange* GetAllocatedRange() { return Range; }
@@ -458,19 +458,27 @@ public:
 	
 	// The returned ranges hide the internal representations so we can change format in the future, 
 	// e.g. store zeroes or 1.0f in some compact fashion or even var int encodings
-	template<typename LeafType>
+	template<Arithmetic T>
 	auto As() const
 	{
-		check(Leaf == ReflectLeaf<LeafType>);
-		if constexpr (std::is_same_v<bool, LeafType>)
+		check(Leaf == ReflectArithmetic<T>);
+		if constexpr (std::is_same_v<bool, T>)
 		{
 			return FBoolRangeView(static_cast<const uint8*>(Data), Num);
 		}
 		else
 		{
-			return TRangeView<LeafType>(static_cast<const LeafType*>(Data), Num);
+			return TRangeView<T>(static_cast<const T*>(Data), Num);
 		}
 	}
+	
+	template<Enumeration T>
+	TRangeView<T> As() const
+	{
+		check(Leaf == ReflectEnum<T>);
+		return TRangeView<T>(static_cast<const T*>(Data), Num);
+	}
+
 };
 
 // Specialized binding for transcoding leaf ranges
@@ -511,9 +519,12 @@ using RangeBind = typename TRangeBind<T>::Type;
 
 struct FMemberBinding
 {
-	FMemberBinding() : InnermostType(FLeafBindType(ELeafBindType::Bool, ELeafWidth::B8)) {}
+	explicit FMemberBinding(uint64 InOffset = 0)
+	: Offset(InOffset)
+	, InnermostType(FLeafBindType(ELeafBindType::Bool, ELeafWidth::B8))
+	{}
 
-	uint64							Offset = 0;
+	uint64							Offset;
 	FMemberBindType					InnermostType;		// Always Leaf or Struct
 	FOptionalSchemaId				InnermostSchema;	// Enum or struct schema
 	TConstArrayView<FRangeBinding>	RangeBindings;		// Non-empty -> Range
@@ -575,16 +586,18 @@ FTypeId IndexAnyType()
 		FTypeId SizeType = { NoId, Ids::IndexTypename(CttiOf<typename RangeBinding::SizeType>::Name) };
 		return Ids::GetIndexer().MakeAnonymousParametricType({InnerType, SizeType});
 	}
-	else if constexpr (std::is_arithmetic_v<T>)
-	{
-		return { NoId, Ids::IndexTypename(CttiOf<T>::Name) };
-	}
 	else 
 	{
 		return IndexStructOrEnumType<CttiOf<T>, Ids>();
 	}
 }
 
+template<Arithmetic T, typename Ids>
+FTypeId IndexAnyType()
+{
+	return { NoId, Ids::IndexTypename(CttiOf<T>::Name) };
+}
+
 template<typename Ids, typename... Ts>
 FTypeId IndexTemplatedType(FTypeId TemplatedType, const std::tuple<Ts...>*)
 {
@@ -608,13 +621,6 @@ FTypeId IndexStructOrEnumType()
 
 ////////////////////////////////////////////////////////////////////////////////////////////////
 
-template<typename Type, class Ids>
-FMemberBindType BindMemberLeaf(FOptionalSchemaId& OutSchema)
-{
-	OutSchema = std::is_enum_v<Type> ? ToOptional(static_cast<FSchemaId>(IndexEnum<Type, Ids>())) : NoId;
-	return FMemberBindType(ReflectLeaf<Type>);
-}
-
 template<typename CustomBinding, class Runtime>
 FStructSchemaId BindCustomStructOnce()
 {
@@ -676,18 +682,24 @@ FMemberBindType BindMemberStruct(FOptionalSchemaId& OutSchema)
 }
 
 template<typename Type, class Ids>
-FMemberBindType BindType(FOptionalSchemaId& OutSchema)
+FMemberBindType BindInnermostType(FOptionalSchemaId& OutSchema)
 {
-	if constexpr (std::is_arithmetic_v<Type> || std::is_enum_v<Type>)
-	{
-		OutSchema = std::is_enum_v<Type> ? ToOptional(static_cast<FSchemaId>(IndexEnum<Type, Ids>())) : NoId;
-		return FMemberBindType(ReflectLeaf<Type>);
-	}
-	else
-	{
-		OutSchema = FOptionalSchemaId(IndexStruct<Type, Ids>());
-		return FMemberBindType(FStructType{EMemberKind::Struct, /* IsDynamic */ 0, /* IsSuper */ 0});
-	}
+	OutSchema = FOptionalSchemaId(IndexStruct<Type, Ids>());
+	return FMemberBindType(FStructType{EMemberKind::Struct, /* IsDynamic */ 0, /* IsSuper */ 0});
+}
+
+template<Arithmetic Type, class Ids>
+FMemberBindType BindInnermostType(FOptionalSchemaId& OutSchema)
+{
+	OutSchema = NoId;
+	return FMemberBindType(ReflectArithmetic<Type>);
+}
+
+template<Enumeration Type, class Ids>
+FMemberBindType BindInnermostType(FOptionalSchemaId& OutSchema)
+{
+	OutSchema = ToOptional(static_cast<FSchemaId>(IndexEnum<Type, Ids>()));
+	return FMemberBindType(ReflectEnum<Type>);
 }
 
 template<typename RangeBinding>
@@ -756,40 +768,39 @@ TConstArrayView<FRangeBinding> GetRangeBindings()
 	}
 }
 
-template<class Var, class Runtime>
-FMemberBinding BindMember()
+template<LeafType Type, class Runtime>
+FMemberBinding BindMember(uint64 Offset)
+{
+	FMemberBinding Out(Offset);
+	Out.InnermostType = BindInnermostType<Type, typename Runtime::Ids>(Out.InnermostSchema);
+	return Out;
+}
+
+template<typename Type, class Runtime>
+FMemberBinding BindMember(uint64 Offset)
 {
 	using Ids = typename Runtime::Ids;
-	using Type = Var::Type;
+	using CustomBinding = typename Runtime::template CustomBindings<Type>::Type;
 
-	FMemberBinding Out;
-	Out.Offset = Var::Offset;
-	if constexpr (std::is_arithmetic_v<Type> || std::is_enum_v<Type>)
+	FMemberBinding Out(Offset);
+	if constexpr (!std::is_void_v<CustomBinding>)
 	{
-		Out.InnermostType = BindMemberLeaf<Type, Ids>(Out.InnermostSchema);
+		Out.InnermostType = BindMemberStruct<Type, CustomBinding, Runtime>(Out.InnermostSchema);
 	}
 	else
 	{
-		using CustomBinding = typename Runtime::template CustomBindings<Type>::Type;
-		if constexpr (!std::is_void_v<CustomBinding>)
+		using RangeBinding = RangeBind<Type>;
+		if constexpr (!std::is_void_v<RangeBinding>)
 		{
-			Out.InnermostType = BindMemberStruct<Type, CustomBinding, Runtime>(Out.InnermostSchema);
+			constexpr uint32 NumRangeBindings = CountRangeBindings<RangeBinding>();
+			using InnermostType = typename TInnerType<RangeBinding, NumRangeBindings>::Type;
+
+			Out.RangeBindings = GetRangeBindings<RangeBinding, NumRangeBindings>();
+			Out.InnermostType = BindInnermostType<InnermostType, Ids>(Out.InnermostSchema);
 		}
 		else
 		{
-			using RangeBinding = RangeBind<Type>;
-			if constexpr (!std::is_void_v<RangeBinding>)
-			{
-				constexpr uint32 NumRangeBindings = CountRangeBindings<RangeBinding>();
-				using InnermostType = typename TInnerType<RangeBinding, NumRangeBindings>::Type;
-
-				Out.RangeBindings = GetRangeBindings<RangeBinding, NumRangeBindings>();
-				Out.InnermostType = BindType<InnermostType, Ids>(Out.InnermostSchema);
-			}
-			else
-			{
-				Out.InnermostType = BindMemberStruct<Type, CustomBinding, Runtime>(Out.InnermostSchema);
-			}
+			Out.InnermostType = BindMemberStruct<Type, CustomBinding, Runtime>(Out.InnermostSchema);
 		}
 	}
 	
@@ -838,7 +849,7 @@ void BindNativeStruct(FSchemaBindings& Out, FStructSchemaId DeclaredId)
 	FMemberBinding MemberBindings[Ctti::NumVars];
 	ForEachVar<Ctti>([&]<class Var>()
 	{ 
-		MemberBindings[Var::Index] = BindMember<Var, Runtime>();
+		MemberBindings[Var::Index] = BindMember<typename Var::Type, Runtime>(Var::Offset);
 	});
 	Out.BindStruct(DeclaredId, MemberBindings);
 }
diff --git a/Engine/Plugins/Experimental/PlainProps/Source/Public/PlainPropsBuild.h b/Engine/Plugins/Experimental/PlainProps/Source/Public/PlainPropsBuild.h
--- a/Engine/Plugins/Experimental/PlainProps/Source/Public/PlainPropsBuild.h
+++ b/Engine/Plugins/Experimental/PlainProps/Source/Public/PlainPropsBuild.h
@@ -141,18 +141,16 @@ struct FTypedRange
 	FBuiltRange* Values = nullptr;
 };
 
-template<typename LeafType, typename SizeType>
+template<Arithmetic T, typename SizeType>
 FMemberSchema MakeLeafRangeSchema()
 {
-	check(std::is_arithmetic_v<LeafType>);
-	return { FMemberType(RangeSizeOf(SizeType{})), ReflectLeaf<LeafType>.Pack(), 1, NoId, nullptr };
+	return { FMemberType(RangeSizeOf(SizeType{})), ReflectArithmetic<T>.Pack(), 1, NoId, nullptr };
 }
 
-template<typename EnumType, typename SizeType>
+template<Enumeration T, typename SizeType>
 FMemberSchema MakeEnumRangeSchema(FEnumSchemaId Schema)
 {
-	check(std::is_enum_v<EnumType>);
-	return { FMemberType(RangeSizeOf(SizeType{})), ReflectLeaf<EnumType>.Pack(), 1, Schema, nullptr };
+	return { FMemberType(RangeSizeOf(SizeType{})), ReflectEnum<T>.Pack(), 1, Schema, nullptr };
 }
 
 inline constexpr FMemberType DefaultStructType = FMemberType(FStructType{EMemberKind::Struct, /* IsDynamic */ 0, /* IsSuper */ 0});
@@ -170,24 +168,24 @@ PLAINPROPS_API FMemberSchema MakeNestedRangeSchema(FScratchAllocator& Scratch, E
 
 [[nodiscard]] PLAINPROPS_API FBuiltRange* CloneLeaves(FScratchAllocator& Scratch, uint64 Num, const void* Data, SIZE_T LeafSize);
 
-template<typename LeafType, typename SizeType>
-[[nodiscard]] FTypedRange BuildLeafRange(FScratchAllocator& Scratch, const LeafType* Values, SizeType Num)
+template<LeafType T, typename SizeType>
+[[nodiscard]] FTypedRange BuildLeafRange(FScratchAllocator& Scratch, const T* Values, SizeType Num)
 {
 	// todo: detect invalid floats
-	return { MakeLeafRangeSchema<LeafType, SizeType>(), CloneLeaves(Scratch, Num, Values, sizeof(LeafType)) };
+	return { MakeLeafRangeSchema<T, SizeType>(), CloneLeaves(Scratch, Num, Values, sizeof(T)) };
 }
 
-template<typename LeafType, typename SizeType>
-[[nodiscard]] FTypedRange BuildLeafRange(FScratchAllocator& Scratch, TConstArrayView<LeafType, SizeType> Values)
+template<LeafType T, typename SizeType>
+[[nodiscard]] FTypedRange BuildLeafRange(FScratchAllocator& Scratch, TConstArrayView<T, SizeType> Values)
 {
 	// todo: detect invalid floats
-	return { MakeLeafRangeSchema<LeafType, SizeType>(), CloneLeaves(Scratch, Values.Num(), Values.GetData(), sizeof(LeafType)) };
+	return { MakeLeafRangeSchema<T, SizeType>(), CloneLeaves(Scratch, Values.Num(), Values.GetData(), sizeof(T)) };
 }
 
-template<typename EnumType, typename SizeType>
-[[nodiscard]] FTypedRange BuildEnumRange(FScratchAllocator& Scratch, FEnumSchemaId Enum, TConstArrayView<EnumType, SizeType> Values)
+template<Enumeration T, typename SizeType>
+[[nodiscard]] FTypedRange BuildEnumRange(FScratchAllocator& Scratch, FEnumSchemaId Enum, TConstArrayView<T, SizeType> Values)
 {
-	return { MakeEnumRangeSchema<EnumType, SizeType>(Enum), CloneLeaves(Scratch, Values.Num(), Values.GetData(), sizeof(EnumType)) };
+	return { MakeEnumRangeSchema<T, SizeType>(Enum), CloneLeaves(Scratch, Values.Num(), Values.GetData(), sizeof(T)) };
 }
 
 [[nodiscard]] inline FTypedRange MakeStructRange(FStructSchemaId Schema, ERangeSizeType SizeType, FBuiltRange* Values )
@@ -223,17 +221,16 @@ struct FBuiltMember
 class FMemberBuilder
 {
 public:
-	template<typename LeafType>
-	void Add(FMemberId Name, LeafType Value)
+	template<Arithmetic T>
+	void Add(FMemberId Name, T Value)
 	{
-		static_assert(std::is_arithmetic_v<LeafType>, "Illegal leaf type");
-		AddLeaf(Name, ReflectLeaf<LeafType>, {}, ValueCast(Value));
+		AddLeaf(Name, ReflectArithmetic<T>, {}, ValueCast(Value));
 	}
 
-	template<typename EnumType>
-	void AddEnum(FMemberId Name, FEnumSchemaId Schema, EnumType Value)
+	template<Enumeration T>
+	void AddEnum(FMemberId Name, FEnumSchemaId Schema, T Value)
 	{
-		AddLeaf(Name, ReflectLeaf<EnumType>, ToOptional(Schema), ValueCast(Value));
+		AddLeaf(Name, ReflectEnum<T>, ToOptional(Schema), ValueCast(Value));
 	}
 	
 	void AddEnum8(FMemberId Name, FEnumSchemaId Schema, uint8 Value)	{ AddLeaf(Name, {ELeafType::Enum, ELeafWidth::B8},  ToOptional(Schema), Value); }
@@ -274,7 +271,7 @@ struct FDenseMemberBuilder
 	{
 		// Todo: Handle enums, ranges and structs
 		FBuiltValue Values[] = { {.Leaf = ValueCast(Head)}, {.Leaf = (ValueCast(Tail))}...  };
-		return BuildHomo(Declaration, FMemberType(ReflectLeaf<T>.Pack()), Values);
+		return BuildHomo(Declaration, FMemberType(ReflectArithmetic<T>.Pack()), Values);
 	}
 
 private:
diff --git a/Engine/Plugins/Experimental/PlainProps/Source/Public/PlainPropsRead.h b/Engine/Plugins/Experimental/PlainProps/Source/Public/PlainPropsRead.h
--- a/Engine/Plugins/Experimental/PlainProps/Source/Public/PlainPropsRead.h
+++ b/Engine/Plugins/Experimental/PlainProps/Source/Public/PlainPropsRead.h
@@ -203,10 +203,17 @@ struct FLeafView
 	FORCEINLINE uint64		AsEnum64() const	{ return As<uint64,	FUnpackedLeafType{ELeafType::Enum, ELeafWidth::B64}>(); }
 
 private:
-	template<typename T, FUnpackedLeafType ExpectedType = ReflectLeaf<T>>
+	template<Arithmetic T, FUnpackedLeafType ExpectedLeaf = ReflectArithmetic<T>>
 	FORCEINLINE T As() const
 	{
-		check(ExpectedType == Leaf);
+		check(ExpectedLeaf == Leaf);
+		return *reinterpret_cast<const T*>(Value.Ptr);
+	}
+
+	template<Enumeration T, FUnpackedLeafType ExpectedLeaf = ReflectEnum<T>>
+	FORCEINLINE T As() const
+	{
+		check(ExpectedLeaf == Leaf);
 		return *reinterpret_cast<const T*>(Value.Ptr);
 	}
 };
@@ -427,13 +434,20 @@ public:
 	TRangeView<char8_t>		AsUtf8() const		{ return As<char8_t>(); }
 	TRangeView<char16_t>	AsUtf16() const		{ return As<char16_t>(); }
 	TRangeView<char32_t>	AsUtf32() const		{ return As<char32_t>(); }
+	
+	template<Arithmetic T>
+	TRangeView<T> As() const
+	{
+		static_assert(!std::is_same_v<T, bool>);
+		check(FUnpackedLeafType(Type, Width) == ReflectArithmetic<T>);
+		return TRangeView<T>(reinterpret_cast<const T*>(Values), NumItems);
+	}
 
-	template<typename LeafType>
-	TRangeView<LeafType> As() const
+	template<Enumeration T>
+	TRangeView<T> As() const
 	{
-		static_assert(!std::is_same_v<LeafType, bool>);
-		check(FUnpackedLeafType(Type, Width) == ReflectLeaf<LeafType>);
-		return TRangeView<LeafType>(reinterpret_cast<const LeafType*>(Values), NumItems);
+		check(FUnpackedLeafType(Type, Width) == ReflectEnum<T>);
+		return TRangeView<T>(reinterpret_cast<const T*>(Values), NumItems);
 	}
 
 	//FEnumFlatRangeView		AsFlatEnums();	// @pre GetType() == ELeafType::Enum
@@ -572,10 +586,16 @@ public:
 	FStructView				GrabStruct();			// @pre PeekKind() == EMemberKind::Struct
 	//FAnyMemberView		GrabAny();				// @pre HasMore()
 
+	// Experimental 
 	// @pre Has N more contiguous members of the expected leaf type
-	template<typename T>
+	template<Arithmetic T>
 	void					GrabLeaves(T* Out, uint32 N);
 
+	// Experimental
+	// @pre Has N more contiguous members of the expected leaf type
+	template<Enumeration T>
+	void					GrabEnums(T* Out, uint32 N);
+
 protected: // for unit tests
 	const FMemberType*		Footer;
 	const FReadBatchId		Batch;					// Needed to resolve schemas
@@ -621,20 +641,15 @@ protected: // for unit tests
 	void					GrabLeaves(void* Out, uint32 Num, SIZE_T NumBytes);
 };
 
-template<typename T>
+template<Arithmetic T>
 void FMemberReader::GrabLeaves(T* Out, uint32 N)
 {
 	if (N)
 	{
-		constexpr FUnpackedLeafType Leaf = ReflectLeaf<T>;
-		if constexpr (Leaf.Type == ELeafType::Bool)
+		if constexpr (std::is_same_v<bool, T>)
 		{
 			GrabBools(Out, N);
 		}
-		else if constexpr (Leaf.Type == ELeafType::Enum)
-		{
-			GrabEnums(Out, N, sizeof(T));
-		}
 		else
 		{
 			GrabLeaves(Out, N, sizeof(T));
@@ -642,6 +657,16 @@ void FMemberReader::GrabLeaves(T* Out, uint32 N)
 	}
 }
 
+template<Enumeration T>
+void FMemberReader::GrabEnums(T* Out, uint32 N)
+{
+	if (N)
+	{
+		GrabLeaves(Out, N, sizeof(T));
+		InnerSchemaIdx += N;
+	}
+}
+
 //////////////////////////////////////////////////////////////////////////
 
 // Hides the inheritance chain and iterates over super members first
diff --git a/Engine/Plugins/Experimental/PlainProps/Source/Public/PlainPropsTypes.h b/Engine/Plugins/Experimental/PlainProps/Source/Public/PlainPropsTypes.h
--- a/Engine/Plugins/Experimental/PlainProps/Source/Public/PlainPropsTypes.h
+++ b/Engine/Plugins/Experimental/PlainProps/Source/Public/PlainPropsTypes.h
@@ -319,6 +319,17 @@ struct FParametricTypeView
 
 //////////////////////////////////////////////////////////////////////////
 
+template<class T>
+concept Arithmetic = std::is_arithmetic_v<T>;
+
+template<class T>
+concept Enumeration = std::is_enum_v<T>;
+
+template<class T>
+concept LeafType = std::is_arithmetic_v<T> || std::is_enum_v<T>;
+
+//////////////////////////////////////////////////////////////////////////
+
 struct FUnpackedLeafType
 {
 	ELeafType Type;
@@ -338,31 +349,34 @@ struct FUnpackedLeafType
 	}
 };
 
+template<Enumeration T>
+inline constexpr FUnpackedLeafType ReflectEnum = { ELeafType::Enum, LeafWidth<sizeof(T)> };
+
 template<typename T>
-constexpr FUnpackedLeafType ReflectEnum()
+constexpr FUnpackedLeafType IllegalLeaf()
 {
-	static_assert(std::is_enum_v<T>);
-	return FUnpackedLeafType(ELeafType::Enum, LeafWidth<sizeof(T)>);
+	static_assert(!sizeof(T), "Unsupported leaf type");
+	return { ELeafType::Bool,	ELeafWidth::B8 };
 }
 
-template<typename T>
-constexpr FUnpackedLeafType ReflectLeaf = ReflectEnum<T>();
-
-template<> inline constexpr FUnpackedLeafType ReflectLeaf<bool>		= { ELeafType::Bool,	ELeafWidth::B8 };
-template<> inline constexpr FUnpackedLeafType ReflectLeaf<int8>		= { ELeafType::IntS,	ELeafWidth::B8 };
-template<> inline constexpr FUnpackedLeafType ReflectLeaf<int16>	= { ELeafType::IntS,	ELeafWidth::B16 };
-template<> inline constexpr FUnpackedLeafType ReflectLeaf<int32>	= { ELeafType::IntS,	ELeafWidth::B32 };
-template<> inline constexpr FUnpackedLeafType ReflectLeaf<int64>	= { ELeafType::IntS,	ELeafWidth::B64 };
-template<> inline constexpr FUnpackedLeafType ReflectLeaf<uint8>	= { ELeafType::IntU,	ELeafWidth::B8 };
-template<> inline constexpr FUnpackedLeafType ReflectLeaf<uint16>	= { ELeafType::IntU,	ELeafWidth::B16 };
-template<> inline constexpr FUnpackedLeafType ReflectLeaf<uint32>	= { ELeafType::IntU,	ELeafWidth::B32 };
-template<> inline constexpr FUnpackedLeafType ReflectLeaf<uint64>	= { ELeafType::IntU,	ELeafWidth::B64 };
-template<> inline constexpr FUnpackedLeafType ReflectLeaf<float>	= { ELeafType::Float,	ELeafWidth::B32 };
-template<> inline constexpr FUnpackedLeafType ReflectLeaf<double>	= { ELeafType::Float,	ELeafWidth::B64 };
-template<> inline constexpr FUnpackedLeafType ReflectLeaf<char>		= { ELeafType::Unicode,	ELeafWidth::B8 };
-template<> inline constexpr FUnpackedLeafType ReflectLeaf<char8_t>	= { ELeafType::Unicode,	ELeafWidth::B8 };
-template<> inline constexpr FUnpackedLeafType ReflectLeaf<char16_t>	= { ELeafType::Unicode,	ELeafWidth::B16 };
-template<> inline constexpr FUnpackedLeafType ReflectLeaf<char32_t>	= { ELeafType::Unicode,	ELeafWidth::B32 };
+template<Arithmetic T>
+inline constexpr FUnpackedLeafType ReflectArithmetic = IllegalLeaf<T>;
+
+template<> inline constexpr FUnpackedLeafType ReflectArithmetic<bool>		= { ELeafType::Bool,	ELeafWidth::B8 };
+template<> inline constexpr FUnpackedLeafType ReflectArithmetic<int8>		= { ELeafType::IntS,	ELeafWidth::B8 };
+template<> inline constexpr FUnpackedLeafType ReflectArithmetic<int16>		= { ELeafType::IntS,	ELeafWidth::B16 };
+template<> inline constexpr FUnpackedLeafType ReflectArithmetic<int32>		= { ELeafType::IntS,	ELeafWidth::B32 };
+template<> inline constexpr FUnpackedLeafType ReflectArithmetic<int64>		= { ELeafType::IntS,	ELeafWidth::B64 };
+template<> inline constexpr FUnpackedLeafType ReflectArithmetic<uint8>		= { ELeafType::IntU,	ELeafWidth::B8 };
+template<> inline constexpr FUnpackedLeafType ReflectArithmetic<uint16>		= { ELeafType::IntU,	ELeafWidth::B16 };
+template<> inline constexpr FUnpackedLeafType ReflectArithmetic<uint32>		= { ELeafType::IntU,	ELeafWidth::B32 };
+template<> inline constexpr FUnpackedLeafType ReflectArithmetic<uint64>		= { ELeafType::IntU,	ELeafWidth::B64 };
+template<> inline constexpr FUnpackedLeafType ReflectArithmetic<float>		= { ELeafType::Float,	ELeafWidth::B32 };
+template<> inline constexpr FUnpackedLeafType ReflectArithmetic<double>		= { ELeafType::Float,	ELeafWidth::B64 };
+template<> inline constexpr FUnpackedLeafType ReflectArithmetic<char>		= { ELeafType::Unicode,	ELeafWidth::B8 };
+template<> inline constexpr FUnpackedLeafType ReflectArithmetic<char8_t>	= { ELeafType::Unicode,	ELeafWidth::B8 };
+template<> inline constexpr FUnpackedLeafType ReflectArithmetic<char16_t>	= { ELeafType::Unicode,	ELeafWidth::B16 };
+template<> inline constexpr FUnpackedLeafType ReflectArithmetic<char32_t>	= { ELeafType::Unicode,	ELeafWidth::B32 };
 
 //////////////////////////////////////////////////////////////////////////
 
diff --git a/Engine/Plugins/Experimental/PlainProps/Source/Public/PlainPropsUeCoreBindings.h b/Engine/Plugins/Experimental/PlainProps/Source/Public/PlainPropsUeCoreBindings.h
--- a/Engine/Plugins/Experimental/PlainProps/Source/Public/PlainPropsUeCoreBindings.h
+++ b/Engine/Plugins/Experimental/PlainProps/Source/Public/PlainPropsUeCoreBindings.h
@@ -110,15 +110,16 @@ namespace PlainProps::UE
 namespace PlainProps::UE
 {
 
-template <typename T>
+template <typename T, class Allocator>
 struct TArrayBinding : public IItemRangeBinding
 {
 	using SizeType = int32;
 	using ItemType = T;
+	using ArrayType = TArray<T, Allocator>;
 
 	virtual void MakeItems(FLoadRangeContext& Ctx) const override
 	{
-		TArray<T>& Array = Ctx.Request.GetRange<TArray<T>>();
+		ArrayType& Array = Ctx.Request.GetRange<ArrayType>();
 		if constexpr (std::is_default_constructible_v<T>)
 		{
 			Array.SetNum(Ctx.Request.NumTotal());
@@ -134,7 +135,7 @@ struct TArrayBinding : public IItemRangeBinding
 
 	virtual void ReadItems(FSaveRangeContext& Ctx) const override
 	{
-		const TArray<T>& Array = Ctx.Request.GetRange<TArray<T>>();
+		const ArrayType& Array = Ctx.Request.GetRange<ArrayType>();
 		Ctx.Items.SetAll(Array.GetData(), static_cast<uint64>(Array.Num()));
 	}
 };
@@ -666,10 +667,10 @@ namespace PlainProps
 template<>
 PLAINPROPS_API void AppendString(FString& Out, const FName& Name);
 
-template<typename T>
-struct TRangeBind<TArray<T>>
+template<typename T, typename Allocator>
+struct TRangeBind<TArray<T, Allocator>>
 {
-	using Type = UE::TArrayBinding<T>;
+	using Type = UE::TArrayBinding<T, Allocator>;
 };
 
 template<>
