PlainProps add compile-time string concat and number -> hex string, might be used for generating unique runtime type names

#rnx

diff --git a/Engine/Plugins/Experimental/PlainProps/Source/Public/PlainPropsStringUtil.h b/Engine/Plugins/Experimental/PlainProps/Source/Public/PlainPropsStringUtil.h
new file mode 100644
--- /dev/null
+++ b/Engine/Plugins/Experimental/PlainProps/Source/Public/PlainPropsStringUtil.h
@@ -0,0 +1,65 @@
+// Copyright Epic Games, Inc. All Rights Reserved.
+
+#pragma once
+
+#include <string_view>
+#include <array>
+#include "HAL/Platform.h" // uint32 etc
+
+namespace PlainProps::Private
+{
+
+inline constexpr void Append(char*& ToIt, std::string_view From)
+{
+    for (char Char : From)
+    {
+        *ToIt++ = Char;
+    }
+}
+
+template <const std::string_view&... Strs>
+struct TConcat
+{
+    inline static constexpr std::size_t Len = (Strs.size() + ... + 0);
+
+    static constexpr std::array<char, Len + 1> Chars = []()
+    {    
+        std::array<char, Len + 1> Out;
+        char* OutIt = Out.data();
+        (Append(OutIt, Strs), ...);
+        *OutIt = '\0';
+        return Out;
+    }();
+
+    inline static constexpr std::string_view Value{Chars.data(), Len};
+};
+
+template <const std::string_view&... Strs>
+inline constexpr std::string_view Concat = TConcat<Strs...>::Value;
+
+template<uint64 N>
+struct THexString
+{
+	inline static constexpr char At(uint32 Idx)
+	{
+		uint64 Nibble =	(N >> 4 * Idx) & uint64(15);
+		return static_cast<char>(Nibble > 9 ? Nibble - 10 + 'A' : Nibble + '0'); 
+	}
+
+	static constexpr char Chars[16] = { At(15), At(14), At(13), At(12), At(11), At(10), At(9), At(8), At(7), At(6), At(5), At(4), At(3), At(2), At(1), At(0) };
+
+	inline static constexpr uint32 CalcLen()
+	{
+		uint32 Idx = 15;
+		for (; Idx && (N >> (Idx * 4)) == 0; --Idx) {}
+		return Idx + 1;
+	}
+
+	inline static constexpr uint32 Len = CalcLen();
+	inline static constexpr std::string_view Value{Chars + 16 - Len, Len};
+};
+
+template <uint64 N>
+inline constexpr std::string_view HexString = THexString<N>::Value;
+
+} // namespace PlainProps::Private
\ No newline at end of file
diff --git a/Engine/Plugins/Experimental/PlainProps/Source/Tests/PlainPropsStringTests.cpp b/Engine/Plugins/Experimental/PlainProps/Source/Tests/PlainPropsStringTests.cpp
new file mode 100644
--- /dev/null
+++ b/Engine/Plugins/Experimental/PlainProps/Source/Tests/PlainPropsStringTests.cpp
@@ -0,0 +1,25 @@
+// Copyright Epic Games, Inc. All Rights Reserved.
+
+#include "PlainPropsStringUtil.h"
+
+namespace PlainProps::Private
+{
+	using namespace std::literals;
+
+	static constexpr std::string_view Hello = "Hello";
+	static constexpr std::string_view Space = " ";
+	static constexpr std::string_view World = "World";
+	static_assert("Hello World"sv == Private::Concat<Hello, Space, World>);
+	
+	static_assert("0"sv == Private::HexString<0>);
+	static_assert("1"sv == Private::HexString<1>);
+	static_assert("9"sv == Private::HexString<9>);
+	static_assert("A"sv == Private::HexString<0xA>);
+	static_assert("F"sv == Private::HexString<0xF>);
+	static_assert("10"sv == Private::HexString<0x10>);
+	static_assert("FF"sv == Private::HexString<0xFF>);
+	static_assert("100"sv == Private::HexString<0x100>);
+	static_assert("FFF"sv == Private::HexString<0xFFF>);
+	static_assert("1000"sv == Private::HexString<0x1000>);
+	static_assert("FEDCBA9876543210"sv == Private::HexString<0xFEDCBA9876543210>);
+}
\ No newline at end of file
