PlainProps some progress on TSetDeltaBinding

#rnx

diff --git a/Engine/Plugins/Experimental/PlainProps/Source/Private/PlainPropsLoad.cpp b/Engine/Plugins/Experimental/PlainProps/Source/Private/PlainPropsLoad.cpp
--- a/Engine/Plugins/Experimental/PlainProps/Source/Private/PlainPropsLoad.cpp
+++ b/Engine/Plugins/Experimental/PlainProps/Source/Private/PlainPropsLoad.cpp
@@ -249,7 +249,7 @@ static void CloneBindingWithReplacedStructIds(const FSchemaId* FromIds, const FS
 	return Memcpy ? FLoadStructPlan(Memcpy.GetValue()) : MakeSchemaLoadPlan(From, To, ToMemberIds, ToStructIds, OutSubsetSchemas);
 }
 
-class FMemberlessDummyBinding : public ICustomBinding
+struct FMemberlessDummyBinding : ICustomBinding
 {
 	virtual void SaveCustom(FMemberBuilder& Dst, const void* Src, const void* Default, const FSaveContext& Ctx) override { check(false); }
 	virtual void LoadCustom(void* Dst, FStructView Src, ECustomLoadMethod Method, const FLoadBatch& Batch) const override { check(false);}
diff --git a/Engine/Plugins/Experimental/PlainProps/Source/Public/PlainPropsBind.h b/Engine/Plugins/Experimental/PlainProps/Source/Public/PlainPropsBind.h
--- a/Engine/Plugins/Experimental/PlainProps/Source/Public/PlainPropsBind.h
+++ b/Engine/Plugins/Experimental/PlainProps/Source/Public/PlainPropsBind.h
@@ -243,9 +243,8 @@ enum ECustomLoadMethod { Construct, Assign };
 /// * custom delta semantics
 /// * other runtime representations than struct/class, e.g. serialize database
 /// * optimization for very common struct
-class ICustomBinding
+struct ICustomBinding
 {
-public:
 	virtual ~ICustomBinding() {}
 	virtual void				SaveCustom(FMemberBuilder& Dst, const void* Src, const void* Default, const FSaveContext& Ctx) = 0;
 	virtual void				LoadCustom(void* Dst, FStructView Src, ECustomLoadMethod Method, const FLoadBatch& Batch) const = 0;
@@ -631,7 +630,7 @@ FStructSchemaId BindCustomStructOnce()
 
 		FBinding()
 		{
-			FTypeId Name = Ids::IndexNativeType(CttiOf<Type>::Name);
+			FTypeId Name = IndexStructOrEnumType<CttiOf<Type>, Ids>();
 			Id = Ids::IndexStruct(Name);
 			CustomBinding::template InitIds<Ids>();
 			Runtime::GetTypes().DeclareStruct(Id, Name, CustomBinding::MemberIds, CustomBinding::Occupancy);
diff --git a/Engine/Plugins/Experimental/PlainProps/Source/Public/PlainPropsUeCoreBindings.h b/Engine/Plugins/Experimental/PlainProps/Source/Public/PlainPropsUeCoreBindings.h
--- a/Engine/Plugins/Experimental/PlainProps/Source/Public/PlainPropsUeCoreBindings.h
+++ b/Engine/Plugins/Experimental/PlainProps/Source/Public/PlainPropsUeCoreBindings.h
@@ -15,7 +15,8 @@
 
 
 PP_NAME_STRUCT(, FName);
-PP_REFLECT_STRUCT_TEMPLATE(, TTuple, void, Key, Value) // Todo handle TTuple and higher arities
+PP_NAME_STRUCT_TEMPLATE(, TSet);
+PP_REFLECT_STRUCT_TEMPLATE(, TTuple, void, Key, Value); // Todo handle TTuple and higher arities
 
 namespace UE::Math
 {
@@ -421,7 +422,7 @@ struct TMapBinding : public TSetBinding<TPair<K, V>, KeyFuncs, SetAllocator>
 //////////////////////////////////////////////////////////////////////////
 
 //TODO: macroify, e.g PP_CUSTOM_BIND(PLAINPROPS_API, FTransform, Transform, Translate, Rotate, Scale)
-struct FTransformBinding : public ICustomBinding
+struct FTransformBinding : ICustomBinding
 {
 	using Type = FTransform;
 	inline static constexpr EMemberPresence Occupancy = EMemberPresence::AllowSparse;
@@ -448,46 +449,57 @@ struct FTransformBinding : public ICustomBinding
 
 //////////////////////////////////////////////////////////////////////////
 
-struct FSetOps
+struct FSetDeltaOps
 {
 	union
 	{
-		struct
-		{
-			FMemberId Add;
-			FMemberId Del;
-			FMemberId Set;
-		};
-
-		FMemberId All[3];
+		FMemberId MemberIds[3] = {};
+		struct { FMemberId Add, Del, Set; } Ops;
 	};
 	
 	template<class Ids>	
-	static const FSetOps& Get()
+	void InitIds()
 	{
-		static FSetOps Ops = {Ids::IndexMember("Add"), Ids::IndexMember("Del"), Ids::IndexMember("Set")};
-		return Ops;
+		static FSetDeltaOps Cache = {.MemberIds = {Ids::IndexMember("Add"), Ids::IndexMember("Del"), Ids::IndexMember("Set")}};
+		*this = Cache;
 	}
 };
 
-template <class Ids, typename T>
-struct TSetDeltaBinding : public ICustomBinding
+template <typename T, typename KeyFuncs, typename SetAllocator>
+struct TSetDeltaBinding : ICustomBinding, FSetDeltaOps
 {
-	using Type = TSet<T>;
+	using Type = TSet<T, KeyFuncs, SetAllocator>;
 	static constexpr EMemberPresence Occupancy = EMemberPresence::AllowSparse;
-	
-	static TConstArrayView<FMemberId> GetMemberIds() { return FSetOps::Get<Ids>().All; }
 
-	virtual void SaveCustom(FMemberBuilder& Dst, const void* Src, const void* Default, const FSaveContext& Ctx) override;
+	void Save(FMemberBuilder& Dst, const Type& Src, const Type* Default, const FSaveContext& Context) const
+	{
+		if (Default)
+		{
+			if (Default->IsEmpty())
+			{
+				// Todo: Add everything
+			}
+			else
+			{
+				// TODO: Range builder for missing items, iterate over defaults elements and check existence in Src
+				// Dst.AddRange(Ops.Del, MissingItems) if non-empty;
 
-	virtual void LoadCustom(void* Dst, FStructView Src, ECustomLoadMethod Method, const FLoadBatch& Batch) const override
+				// TODO: Range builder for added items, iterate over defaults elements and check existence in Src
+			}
+		}
+		else
+		{
+			// Dst.AddRange(Ops.Set, );
+		}
+	}
+
+	inline void Load(Type& Dst, FStructView Src, ECustomLoadMethod Method, const FLoadBatch& Batch) const
 	{
-		Type& Out = *static_cast<Type*>(Dst);
 		FMemberReader Members(Src);
 
 		if (Method == ECustomLoadMethod::Construct)
 		{
-			::new (Dst) Type;
+			::new (&Dst) Type;
 		}
 				
 		if (!Members.HasMore())
@@ -495,40 +507,37 @@ struct TSetDeltaBinding : public ICustomBinding
 			return;
 		}
 
-		FSetOps Ops = FSetOps::Get<Ids>();
 		FMemberId Name = Members.PeekName().Get();
 		FRangeView Items = Members.GrabRange();
 		int32 NumItems = static_cast<int32>(Items.Num());
 		if (Name == Ops.Set)
 		{
-			Out.Empty(NumItems);
-			AddItems(Out, Items);
+			Dst.Empty(NumItems);
+			AddItems(Dst, Items, Batch);
 		}
 		else if (Name == Ops.Add)
 		{
-			Out.Reserve(Out.Num() + NumItems);
-			AddItems(Out, Items);
+			Dst.Reserve(Dst.Num() + NumItems);
+			AddItems(Dst, Items, Batch);
 		}
 		else if (Name == Ops.Del)
 		{
-			RemoveItems(Out, NumItems);
+			RemoveItems(Dst, Items);
 		
 			if (Members.HasMore())
 			{
 				check(Members.PeekName() == Ops.Add);
 				Items = Members.GrabRange();
-				Out.Reserve(Out.Num() + static_cast<int32>(Items.Num()));
-				AddItems(Out, Items);
+				Dst.Reserve(Dst.Num() + static_cast<int32>(Items.Num()));
+				AddItems(Dst, Items, Batch);
 			}
 		}
 		
 		check(!Members.HasMore());
 	}
 
-	virtual bool DiffCustom(const void* StructA, const void* StructB) const override
+	inline static bool Diff(const Type& A, const Type& B)
 	{
-		const Type& A = *static_cast<const Type*>(StructA);
-		const Type& B = *static_cast<const Type*>(StructB);
 		if (A.Num() != B.Num())
 		{
 			return false;
@@ -545,11 +554,11 @@ struct TSetDeltaBinding : public ICustomBinding
 		return true;
 	}
 
-	void AddItems(Type& Out, FRangeView Items, const FLoadBatch& Batch)
+	static void AddItems(Type& Out, FRangeView Items, const FLoadBatch& Batch)
 	{
 		check(!Items.IsEmpty());
 
-		if constexpr (std::is_arithmetic_v<T> || std::is_enum_v<T>)
+		if constexpr (LeafType<T>)
 		{
 			for (T Item : Items.AsLeaves().As<T>())
 			{
@@ -561,11 +570,11 @@ struct TSetDeltaBinding : public ICustomBinding
 			FStructRangeView Structs = Items.AsStructs();
 			for (FStructView Item : Structs)
 			{
-				if constexpr (std::is_default_constructible_v<T>())
+				if constexpr (std::is_default_constructible_v<T>)
 				{
 					T Tmp;
 					PlainProps::LoadStruct(&Tmp, Item, Batch);	
-					Out.Emplace(MoveTemp(Item));
+					Out.Emplace(MoveTemp(Tmp));
 				}
 				else
 				{
@@ -577,22 +586,32 @@ struct TSetDeltaBinding : public ICustomBinding
 				}
 			}
 		}
-		else if constexpr (std::is_default_constructible_v<T>())
+		else // Nested range
 		{
-			using Binding = RangeBind<T>;
-			const IItemRangeBinding* Bindings[] = {};// ... generate somehow ... };
-			T Tmp;
-			for (FRangeView Item : Items.AsRanges())
+			using RangeBinding = RangeBind<T>;
+			if constexpr (std::is_default_constructible_v<T> && !std::is_void_v<RangeBinding>)
 			{
-				LoadRange(&Tmp, Item, Bindings, Batch);
-				Out.Emplace(MoveTemp(*Tmp));
+				static constexpr ERangeSizeType MaxSize = RangeSizeOf(typename RangeBinding::SizeType{});
+				//const IItemRangeBinding* Bindings[] = {};// TODO ... generate somehow ... ;
+				T Tmp;
+				for (FRangeView Item : Items.AsRanges())
+				{
+				//	LoadRange(&Tmp, Item, MaxSize, Bindings, Batch);
+				//	Out.Emplace(MoveTemp(*Tmp));
+				}
+			}
+			else
+			{
+				check(Items.IsNestedRange());
+				checkf(std::is_default_constructible_v<T>, TEXT("Ranges must be default-constructible"));
+				checkf(!std::is_void_v<RangeBinding>, TEXT("Inner range type unbound (no TRangeBind specialization)"));
 			}
 		}
-		else
-		{
-			check(Items.IsNestedRange());
-			checkf(false, TEXT("Ranges must be default-constructible"));
-		}
+	}
+
+	static void RemoveItems(Type& Out, FRangeView Items)
+	{
+		// TODO copy paste from AddItems or template AddItems -> ApplyItems<EOp::Add/Del>
 	}
 };
 
@@ -612,7 +631,7 @@ struct TSetDeltaBinding : public ICustomBinding
 //	};
 //
 //template <typename... Ts>
-//struct TVariantBinding : public ICustomBinding
+//struct TVariantBinding : ICustomBinding
 //{
 //	using VariantType = TVariant<Ts...>;
 //
diff --git a/Engine/Plugins/Experimental/PlainProps/Source/Tests/PlainPropsUeCoreTests.cpp b/Engine/Plugins/Experimental/PlainProps/Source/Tests/PlainPropsUeCoreTests.cpp
--- a/Engine/Plugins/Experimental/PlainProps/Source/Tests/PlainPropsUeCoreTests.cpp
+++ b/Engine/Plugins/Experimental/PlainProps/Source/Tests/PlainPropsUeCoreTests.cpp
@@ -26,6 +26,7 @@ static TIdIndexer<FName>	GNames;
 static FDeclarations		GTypes(/* debug */ GNames);
 static FSchemaBindings		GSchemas(/* debug */ GNames);
 static FCustomBindings		GCustoms(/* debug */ GNames);
+static FCustomBindings		GDeltaCustoms(/* debug */ GNames, /* base*/ &GCustoms);
 
 struct FIds
 {
@@ -54,6 +55,34 @@ struct FDefaultRuntime
 	static FCustomBindings&			GetCustoms()		{ return GCustoms; }
 };
 
+
+////////////////////////////////////////////////////////////////////////////////////////////////
+
+template<class T>
+struct TCustomDeltaBind
+{
+	using Type = void;
+};
+
+template <typename T, typename KeyFuncs, typename SetAllocator>
+struct TCustomDeltaBind<TSet<T, KeyFuncs, SetAllocator>>
+{
+	using Type = UE::TSetDeltaBinding<T, KeyFuncs, SetAllocator>;
+};
+
+//template <typename K, typename V, typename SetAllocator, typename KeyFuncs>
+//struct TCustomDeltaBind<TMap<K, V, SetAllocator, KeyFuncs>>
+//{
+//	using Type = UE::TMapDeltaBinding<K, V, SetAllocator, KeyFuncs>;
+//};
+
+struct FDeltaRuntime : FDefaultRuntime
+{
+	template<class T> using CustomBindings = TCustomDeltaBind<T>;
+
+	static FCustomBindings&			GetCustoms()		{ return GDeltaCustoms; }
+};
+
 ////////////////////////////////////////////////////////////////////////////////////////////////
 
 template<typename Enum, EEnumMode Mode, class Runtime = FDefaultRuntime>
@@ -126,12 +155,12 @@ struct FNameDeclaration
 	}
 };
 
-struct FTestCustomBinding : public ICustomBinding
+struct FTestCustomBinding : ICustomBinding
 {
 	virtual FStructSchemaId GetId() const = 0;
 };
 
-struct FNameBinding : public FTestCustomBinding
+struct FNameBinding : FTestCustomBinding
 {
 	virtual void SaveCustom(FMemberBuilder& Dst, const void* Src, const void*, const FSaveContext& Ctx) override
 	{
@@ -160,8 +189,6 @@ struct FNameBinding : public FTestCustomBinding
 
 ////////////////////////////////////////////////////////////////////////////////////////////////
 
-inline constexpr uint32 Magics[] = { 0xFEEDF00D, 0xABCD1234, 0xDADADAAA, 0x99887766, 0xF0F1F2F3 , 0x00112233};
-
 class FBatchSaver
 {
 public:
@@ -221,6 +248,8 @@ TConstArrayView<T> GrabNumAndArray(/* in-out */ FByteReader& It)
 	return MakeArrayView(reinterpret_cast<const T*>(It.GrabBytes(Num * sizeof(T))), Num);
 }
 
+inline constexpr uint32 Magics[] = { 0xFEEDF00D, 0xABCD1234, 0xDADADAAA, 0x99887766, 0xF0F1F2F3 , 0x00112233};
+
 TArray64<uint8> FBatchSaver::Write() const
 {
 	// Build partial schemas
@@ -588,7 +617,7 @@ PP_REFLECT_STRUCT(PlainProps::UE::Test, FDelta, void, A, B, C, D, E);
 //	int Id = 0;
 //};
 //
-//struct FObjectReferenceBinding : public ICustomBinding
+//struct FObjectReferenceBinding : ICustomBinding
 //{
 //	using Type = FObject*;
 //
@@ -882,6 +911,12 @@ TEST_CASE_NAMED(FPlainPropsUeCoreTest, "System::Core::Serialization::PlainProps:
 			});
 	}
 
+	SECTION("TSetDelta")
+	{
+		TScopedStructBinding<FInt> Int;
+		//TScopedStructBinding<FSets, EMemberPresence::AllowSparse, FDeltaRuntime> Sets;
+	}
+
 	SECTION("Transform")
 	{
 		TScopedStructBinding<FVector, EMemberPresence::RequireAll> Vector;
