PlainProps add TTypename mechanism to support multiple runtime representations of a struct, range type-erasure, custom struct names and type aliases

Implemented allocator type erasure for TArray, TSet, TMap
Replaced naming hack that added fake CTTI for arithmetic types with TTypename<Arithmetic>

#rnx

diff --git a/Engine/Plugins/Experimental/PlainProps/Source/Private/PlainPropsBind.cpp b/Engine/Plugins/Experimental/PlainProps/Source/Private/PlainPropsBind.cpp
--- a/Engine/Plugins/Experimental/PlainProps/Source/Private/PlainPropsBind.cpp
+++ b/Engine/Plugins/Experimental/PlainProps/Source/Private/PlainPropsBind.cpp
@@ -147,37 +147,37 @@ void* FLeafRangeAllocator::Allocate(FUnpackedLeafType Leaf, uint64 Num)
 
 ////////////////////////////////////////////////////////////////////////////////////////////////
 
-void FCustomBindings::BindStruct(FStructSchemaId Id, ICustomBinding& Binding)
+void FCustomBindings::BindStruct(FStructSchemaId BindId, FStructSchemaId DeclId, ICustomBinding& Binding)
 {
-	checkf(!Find(Id), TEXT("'%s' already bound"), *Debug.Print(Id));
-	Entries.Emplace(Id, &Binding);
+	checkf(!Find(BindId), TEXT("'%s' already bound"), *Debug.Print(BindId));
+	Entries.Emplace(BindId, DeclId, &Binding);
 }
 
-void FCustomBindings::DropStruct(FStructSchemaId Id)
+void FCustomBindings::DropStruct(FStructSchemaId BindId)
 {
-	for (FEntry& Entry : Entries)
+	for (FCustomBindingEntry& Entry : Entries)
 	{
-		if (Entry.Id == Id)
+		if (Entry.BindId == BindId)
 		{
 			Entries.RemoveAtSwap(&Entry - Entries.GetData(), EAllowShrinking::No);
 			return;
 		}
 	}
 	
-	checkf(false, TEXT("'%s' unbound"), *Debug.Print(Id));
+	checkf(false, TEXT("'%s' unbound"), *Debug.Print(BindId));
 }
 
-ICustomBinding*	FCustomBindings::Find(FStructSchemaId Id) const
+FCustomBindingEntry	FCustomBindings::Find(FStructSchemaId BindId) const
 {
-	for (const FEntry& Entry : Entries)
+	for (FCustomBindingEntry Entry : Entries)
 	{
-		if (Entry.Id == Id)
+		if (Entry.BindId == BindId)
 		{
-			return Entry.Binding;
+			return Entry;
 		}
 	}
 
-	return Base ? Base->Find(Id) : nullptr;
+	return Base ? Base->Find(BindId) : FCustomBindingEntry{};
 }
 
 ////////////////////////////////////////////////////////////////////////////////////////////////
@@ -205,10 +205,10 @@ static uint16 CountRanges(TConstArrayView<FMemberBinding> Members)
 	return IntCastChecked<uint16>(Out);
 }
 
-void FSchemaBindings::BindStruct(FStructSchemaId Id, TConstArrayView<FMemberBinding> Members)
+void FSchemaBindings::BindStruct(FStructSchemaId BindId, FStructSchemaId DeclId, TConstArrayView<FMemberBinding> Members)
 {
 	// Make header, allocate and copy header
-	FSchemaBinding Header = { IntCastChecked<uint16>(Members.Num()), CountInnerSchemas(Members), CountRanges(Members) };
+	FSchemaBinding Header = { DeclId, IntCastChecked<uint16>(Members.Num()), CountInnerSchemas(Members), CountRanges(Members) };
 	FSchemaBinding* Schema = new (FMemory::MallocZeroed(Header.CalculateSize())) FSchemaBinding {Header};
 
 	// Write footer
@@ -232,24 +232,34 @@ void FSchemaBindings::BindStruct(FStructSchemaId Id, TConstArrayView<FMemberBind
 	}
 
 	// Bind
-	if (Id.Idx >= static_cast<uint32>(Bindings.Num()))
+	if (BindId.Idx >= static_cast<uint32>(Bindings.Num()))
 	{
-		Bindings.SetNum(Id.Idx + 1);
+		Bindings.SetNum(BindId.Idx + 1);
 	}
-	checkf(!Bindings[Id.Idx], TEXT("'%s' already bound"), *Debug.Print(Id));
-	Bindings[Id.Idx].Reset(Schema);
+	checkf(!Bindings[BindId.Idx], TEXT("'%s' already bound"), *Debug.Print(BindId));
+	Bindings[BindId.Idx].Reset(Schema);
 }
 
-const FSchemaBinding& FSchemaBindings::GetStruct(FStructSchemaId Id) const
+const FSchemaBinding* FSchemaBindings::FindStruct(FStructSchemaId BindId) const
 {
-	checkf(Id.Idx < (uint32)Bindings.Num() && Bindings[Id.Idx], TEXT("'%s' is unbound"), *Debug.Print(Id));
-	return *Bindings[Id.Idx].Get();
+	return BindId.Idx < (uint32)Bindings.Num() ? Bindings[BindId.Idx].Get() : nullptr;
 }
 
-void FSchemaBindings::DropStruct(FStructSchemaId Id)
+const FSchemaBinding& FSchemaBindings::GetStruct(FStructSchemaId BindId) const
 {
-	checkf(Id.Idx < (uint32)Bindings.Num() && Bindings[Id.Idx], TEXT("'%s' is unbound"), *Debug.Print(Id));
-	Bindings[Id.Idx].Reset();
+	checkf(BindId.Idx < (uint32)Bindings.Num() && Bindings[BindId.Idx], TEXT("'%s' is unbound"), *Debug.Print(BindId));
+	return *Bindings[BindId.Idx].Get();
+}
+
+void FSchemaBindings::DropStruct(FStructSchemaId BindId)
+{
+	checkf(BindId.Idx < (uint32)Bindings.Num() && Bindings[BindId.Idx], TEXT("'%s' is unbound"), *Debug.Print(BindId));
+	Bindings[BindId.Idx].Reset();
+}
+
+FStructSchemaId FSchemaBindings::GetDeclId(FStructSchemaId BindId) const
+{
+	return GetStruct(BindId).DeclId;
 }
 
 //////////////////////////////////////////////////////////////////////////
diff --git a/Engine/Plugins/Experimental/PlainProps/Source/Private/PlainPropsBuildSchema.cpp b/Engine/Plugins/Experimental/PlainProps/Source/Private/PlainPropsBuildSchema.cpp
--- a/Engine/Plugins/Experimental/PlainProps/Source/Private/PlainPropsBuildSchema.cpp
+++ b/Engine/Plugins/Experimental/PlainProps/Source/Private/PlainPropsBuildSchema.cpp
@@ -72,15 +72,14 @@ struct FEnumSchemaBuilder
 
 //////////////////////////////////////////////////////////////////////////
 
-FSchemasBuilder::FSchemasBuilder(const FDeclarations& Types, FScratchAllocator& InScratch)
-: FSchemasBuilder(Types.GetStructs(), Types.GetEnums(), Types.GetDebug(), InScratch)
-{
-
-}
+FSchemasBuilder::FSchemasBuilder(const FDeclarations& Types, const IStructBindIds& InBindIds, FScratchAllocator& InScratch)
+: FSchemasBuilder(Types.GetStructs(), Types.GetEnums(), InBindIds, Types.GetDebug(), InScratch)
+{}
 
-FSchemasBuilder::FSchemasBuilder(FStructDeclarations InStructs, FEnumDeclarations InEnums, const FDebugIds& InDebug, FScratchAllocator& InScratch)
+FSchemasBuilder::FSchemasBuilder(FStructDeclarations InStructs, FEnumDeclarations InEnums, const IStructBindIds& InBindIds, const FDebugIds& InDebug, FScratchAllocator& InScratch)
 : DeclaredStructs(InStructs)
 , DeclaredEnums(InEnums)
+, BindIds(InBindIds)
 , Scratch(InScratch)
 , Debug(InDebug)
 {
@@ -109,16 +108,20 @@ FORCEINLINE FEnumSchemaBuilder&	FSchemasBuilder::NoteEnum(FEnumSchemaId Id)
 	return GetOrEmplace(EnumIndices[Id.Idx], Enums, *DeclaredEnums[Id.Idx], Id);
 }
 
-FORCEINLINE FStructSchemaBuilder& FSchemasBuilder::NoteStruct(FStructSchemaId Id)
+FORCEINLINE FStructSchemaBuilder& FSchemasBuilder::NoteStruct(FStructSchemaId BindId)
 {
 	checkf(!bBuilt, TEXT("Noted new members after building"));
-	checkf(Id.Idx < uint32(DeclaredStructs.Num()) && DeclaredStructs[Id.Idx], TEXT("Undeclared struct '%s' noted"), *Debug.Print(Id));
-	return GetOrEmplace(StructIndices[Id.Idx], Structs, *DeclaredStructs[Id.Idx], *this);
+
+	FStructSchemaId DeclId	= BindId.Idx < uint32(DeclaredStructs.Num()) && DeclaredStructs[BindId.Idx] 
+							? BindId : BindIds.GetDeclId(BindId);
+
+	checkf(DeclId.Idx < uint32(DeclaredStructs.Num()) && DeclaredStructs[DeclId.Idx], TEXT("Undeclared struct '%s' noted"), *Debug.Print(DeclId));
+	return GetOrEmplace(StructIndices[DeclId.Idx], Structs, *DeclaredStructs[DeclId.Idx], *this);
 }
 
-void FSchemasBuilder::NoteStructAndMembers(FStructSchemaId Id, const FBuiltStruct& Struct)
+void FSchemasBuilder::NoteStructAndMembers(FStructSchemaId BindId, const FBuiltStruct& Struct)
 {
-	NoteStruct(Id).NoteMembersRecursively(Struct);
+	NoteStruct(BindId).NoteMembersRecursively(Struct);
 }
 
 FBuiltSchemas FSchemasBuilder::Build()
diff --git a/Engine/Plugins/Experimental/PlainProps/Source/Private/PlainPropsDeclare.cpp b/Engine/Plugins/Experimental/PlainProps/Source/Private/PlainPropsDeclare.cpp
--- a/Engine/Plugins/Experimental/PlainProps/Source/Private/PlainPropsDeclare.cpp
+++ b/Engine/Plugins/Experimental/PlainProps/Source/Private/PlainPropsDeclare.cpp
@@ -1,6 +1,7 @@
 // Copyright Epic Games, Inc. All Rights Reserved.
 
 #include "PlainPropsDeclare.h"
+#include "Algo/Compare.h"
 #include "Containers/Set.h"
 
 namespace PlainProps
@@ -43,20 +44,31 @@ void CopyItems(T* It, TConstArrayView<T> Items)
 	}
 }
 
-void FDeclarations::DeclareStruct(FStructSchemaId Id, FTypeId Type, TConstArrayView<FMemberId> MemberOrder, EMemberPresence Occupancy, FOptionalStructSchemaId Super)
+void FDeclarations::DeclareStruct(FStructSchemaId DeclId, FTypeId Type, TConstArrayView<FMemberId> MemberOrder, EMemberPresence Occupancy, FOptionalStructSchemaId Super)
 {
-	if (static_cast<int32>(Id.Idx) >= DeclaredStructs.Num())
+	if (static_cast<int32>(DeclId.Idx) >= DeclaredStructs.Num())
 	{
-		DeclaredStructs.SetNum(Id.Idx + 1);
+		DeclaredStructs.SetNum(DeclId.Idx + 1);
 	}
 
-	TUniquePtr<FStructDeclaration>& Ptr = DeclaredStructs[Id.Idx];
-	checkf(!Ptr, TEXT("'%s' is already declared"), *Debug.Print(Id));
+	TUniquePtr<FStructDeclaration>& Ptr = DeclaredStructs[DeclId.Idx];
+	if (Ptr)
+	{
+		check(DeclId == Ptr->Id);
+		check(Type == Ptr->Type);
+		check(Super == Ptr->Super);
+		check(Occupancy == Ptr->Occupancy);
+		check(Algo::Compare(MemberOrder, Ptr->GetMemberOrder()));
 
-	FStructDeclaration Header{Id, Type, Super, Occupancy, IntCastChecked<uint16>(MemberOrder.Num())};
-	void* Data = FMemory::Malloc(sizeof(FStructDeclaration) + MemberOrder.Num() * MemberOrder.GetTypeSize());
-	Ptr.Reset(new (Data) FStructDeclaration(Header));
-	CopyItems(Ptr->MemberOrder, MemberOrder);
+		Ptr->RefCount++;
+	}
+	else
+	{
+		FStructDeclaration Header{1, DeclId, Type, Super, Occupancy, IntCastChecked<uint16>(MemberOrder.Num())};
+		void* Data = FMemory::Malloc(sizeof(FStructDeclaration) + MemberOrder.Num() * MemberOrder.GetTypeSize());
+		Ptr.Reset(new (Data) FStructDeclaration(Header));
+		CopyItems(Ptr->MemberOrder, MemberOrder);
+	}
 }
 
 void FDeclarations::DeclareEnum(FEnumSchemaId Id, FTypeId Type, EEnumMode Mode, ELeafWidth Width, TConstArrayView<FEnumerator> Enumerators)
@@ -77,6 +89,18 @@ void FDeclarations::DeclareEnum(FEnumSchemaId Id, FTypeId Type, EEnumMode Mode,
 	ValidateDeclaration(*Ptr);
 }
 
+void FDeclarations::DropStructRef(FStructSchemaId DeclId)
+{
+	Check(DeclId);
+	
+	TUniquePtr<FStructDeclaration>& Ptr = DeclaredStructs[DeclId.Idx];
+	Ptr->RefCount--;
+	if (Ptr->RefCount == 0)
+	{
+		Ptr.Reset();
+	}
+}
+
 #if DO_CHECK
 void FDeclarations::Check(FEnumSchemaId Id) const
 {
diff --git a/Engine/Plugins/Experimental/PlainProps/Source/Private/PlainPropsLoad.cpp b/Engine/Plugins/Experimental/PlainProps/Source/Private/PlainPropsLoad.cpp
--- a/Engine/Plugins/Experimental/PlainProps/Source/Private/PlainPropsLoad.cpp
+++ b/Engine/Plugins/Experimental/PlainProps/Source/Private/PlainPropsLoad.cpp
@@ -4,6 +4,7 @@
 #include "PlainPropsBind.h"
 #include "PlainPropsInternalFormat.h"
 #include "PlainPropsInternalRead.h"
+#include "Containers/Set.h"
 #include "Misc/Optional.h"
 #include <type_traits>
 
@@ -124,6 +125,57 @@ static bool HasDifferentSupers(const FStructSchema& From, const FSchemaBinding&
 
 ////////////////////////////////////////////////////////////////////////////
 
+// Used to create an additional load plan, beyond the saved struct schema ids
+struct FLoadIdMapping
+{
+	FStructSchemaId SaveId; // ~ Batch decl id, index into saved schemas and load plans
+	FStructSchemaId LoadId; // ~ Batch bind id, index into load plans
+	FStructSchemaId BindId; // Runtime bind id
+};
+
+// Helps load type-erased structs with ExplicitBindName by allocating new load-time struct ids
+class FLoadIdBinder
+{
+public:
+	FLoadIdBinder(TConstArrayView<FStructSchemaId> RuntimeDeclIds)
+	: DeclIds(RuntimeDeclIds)
+	, NextLoadIdx(static_cast<uint32>(RuntimeDeclIds.Num()))
+	{}
+
+	FStructSchemaId					BindLoadId(FStructSchemaId SaveId, FStructSchemaId BindId)
+	{
+		FStructSchemaId DeclId = DeclIds[SaveId.Idx];
+		return BindId == DeclId ? SaveId : MapLoadId(SaveId, BindId);
+	}
+	FLoadIdMapping					GetMapping(int32 Idx) const { return Mappings[Idx]; }
+	int32							NumMappings() const { return Mappings.Num(); }
+
+private:
+	TConstArrayView<FStructSchemaId>	DeclIds;
+	uint32								NextLoadIdx;
+	TArray<FLoadIdMapping>				Mappings;
+
+	FStructSchemaId MapLoadId(FStructSchemaId SaveId, FStructSchemaId BindId)
+	{
+		for (FLoadIdMapping Mapping : Mappings)
+		{
+			if (Mapping.BindId == BindId)
+			{
+				check(Mapping.SaveId == SaveId);
+				return Mapping.LoadId;
+			}
+		}
+
+		FLoadIdMapping& Mapping = Mappings.AddDefaulted_GetRef();
+		Mapping.SaveId = SaveId;
+		Mapping.LoadId = { NextLoadIdx++ };
+		Mapping.BindId = BindId;
+		return Mapping.LoadId;
+	}
+};
+
+////////////////////////////////////////////////////////////////////////////
+
 struct FLoadBatch
 {
 	FReadBatchId			ReadId; // Needed to access schemas for custom struct loading
@@ -146,20 +198,24 @@ static void CopyMemberBinding(FLeafMemberBinding Binding, /* in-out */ const FSc
 	Out.AddMember(Binding.Leaf.Pack(), static_cast<uint32>(Binding.Offset));
 }
 
-static void CopyMemberBinding(FStructMemberBinding Binding, /* in-out */ const FSchemaId*& InnerSchemaIt,FMemberBinder& Out)
+// TODO: Make all below free functions FLoadPlanner member functions and make LoadIds a member variable
+
+static void CopyMemberBinding(FStructMemberBinding Binding, /* in-out */ const FSchemaId*& InnerSchemaIt, FLoadIdBinder& LoadIds, FMemberBinder& Out)
 {
+	FStructSchemaId LoadId = LoadIds.BindLoadId(static_cast<FStructSchemaId>(*InnerSchemaIt), Binding.Id);
 	Out.AddMember(FMemberBindType(Binding.Type), static_cast<uint32>(Binding.Offset));
-	Out.AddInnerSchema(*InnerSchemaIt);
+	Out.AddInnerSchema(LoadId);
 	++InnerSchemaIt;
 }
 
-static void CopyMemberBinding(FRangeMemberBinding Binding, /* in-out */ const FSchemaId*& InnerSchemaIt, FMemberBinder& Out)
+static void CopyMemberBinding(FRangeMemberBinding Binding, /* in-out */ const FSchemaId*& InnerSchemaIt, FLoadIdBinder& LoadIds, FMemberBinder& Out)
 {
 	FMemberBindType InnermostType = Binding.InnerTypes[Binding.NumRanges - 1];
 	Out.AddRange(MakeArrayView(Binding.RangeBindings, Binding.NumRanges), InnermostType, static_cast<uint32>(Binding.Offset));
 	if (InnermostType.IsStruct())
 	{
-		Out.AddInnerSchema(*InnerSchemaIt);
+		FStructSchemaId LoadId = LoadIds.BindLoadId(static_cast<FStructSchemaId>(*InnerSchemaIt), static_cast<FStructSchemaId>(Binding.InnermostSchema.Get()));
+		Out.AddInnerSchema(LoadId);
 		++InnerSchemaIt;
 	}
 	else
@@ -168,24 +224,24 @@ static void CopyMemberBinding(FRangeMemberBinding Binding, /* in-out */ const FS
 	}
 }
 
-static void CopyMemberBinding(/* in-out */ FMemberVisitor& BindIt, /* in-out */ const FSchemaId*& InnerSchemaIt, FMemberBinder& Out)
+static void CopyMemberBinding(/* in-out */ FMemberVisitor& BindIt, /* in-out */ const FSchemaId*& InnerSchemaIt, FLoadIdBinder& LoadIds, FMemberBinder& Out)
 {
 	switch (BindIt.PeekKind())
 	{
-		case EMemberKind::Leaf:		CopyMemberBinding(BindIt.GrabLeaf(), InnerSchemaIt, Out);	break;
-		case EMemberKind::Range:	CopyMemberBinding(BindIt.GrabRange(), InnerSchemaIt, Out);	break;
-		case EMemberKind::Struct:	CopyMemberBinding(BindIt.GrabStruct(), InnerSchemaIt, Out);	break;
-		default:					check(false);												break;
+		case EMemberKind::Leaf:		CopyMemberBinding(BindIt.GrabLeaf(), InnerSchemaIt, Out);				break;
+		case EMemberKind::Range:	CopyMemberBinding(BindIt.GrabRange(), InnerSchemaIt, LoadIds, Out);		break;
+		case EMemberKind::Struct:	CopyMemberBinding(BindIt.GrabStruct(), InnerSchemaIt, LoadIds, Out);	break;
+		default:					check(false);															break;
 	}
 }
 
-static void CreateSubsetBindingWithoutEnumIds(const FStructSchema& From, const FSchemaBinding& To, TConstArrayView<FMemberId> ToNames,  uint16 NumEnums, SubsetByteArray& Out)
+static void CreateSubsetBindingWithoutEnumIds(const FStructSchema& From, const FSchemaBinding& To, TConstArrayView<FMemberId> ToNames, uint16 NumEnums, FLoadIdBinder& LoadIds, SubsetByteArray& Out)
 {
 	check(To.NumMembers == ToNames.Num());
 	check(To.NumMembers >= From.NumMembers);
 
 	int32 OutPos = Out.Num();
-	FSchemaBinding Header = { From.NumMembers, From.NumInnerSchemas - NumEnums, From.NumRangeTypes };
+	FSchemaBinding Header = { To.DeclId, From.NumMembers, From.NumInnerSchemas - NumEnums, From.NumRangeTypes };
 	Out.AddUninitialized(Header.CalculateSize());
 	FSchemaBinding* Schema = new (&Out[OutPos]) FSchemaBinding {Header};
 	
@@ -199,26 +255,33 @@ static void CreateSubsetBindingWithoutEnumIds(const FStructSchema& From, const F
 			ToIt.SkipMember();
 		}
 
-		CopyMemberBinding(/* in-out */ ToIt, /* in-out */ InnerSchemaIt, /* out */ Footer);	
+		CopyMemberBinding(/* in-out */ ToIt, /* in-out */ InnerSchemaIt, /* in-out */ LoadIds,  /* out */ Footer);	
 	}
 	check(InnerSchemaIt == From.GetInnerSchemas() + From.NumInnerSchemas);
 }
 
-static void CloneBindingWithReplacedStructIds(const FSchemaId* FromIds, const FSchemaBinding& To, SubsetByteArray& Out)
+static void CloneBindingWithReplacedStructIds(const FSchemaId* FromIds, const FSchemaBinding& To, FLoadIdBinder& LoadIds, SubsetByteArray& Out)
 {
 	uint32 Size = To.CalculateSize();
 	Out.AddUninitialized(Size);
 	FSchemaBinding* Schema = reinterpret_cast<FSchemaBinding*>(&Out[Out.Num() - Size]);
 	FMemory::Memcpy(Schema, &To, Size);
-	FMemory::Memcpy(const_cast<FSchemaId*>(Schema->GetInnerSchemas()), FromIds, To.NumInnerSchemas * sizeof(FSchemaId));
+
+	const FStructSchemaId* SaveIdIt = static_cast<const FStructSchemaId*>(FromIds);
+	FStructSchemaId* OutIds = static_cast<FStructSchemaId*>(const_cast<FSchemaId*>(Schema->GetInnerSchemas()));
+	for (FStructSchemaId& OutId : MakeArrayView(OutIds, To.NumInnerSchemas))
+	{
+		FStructSchemaId MemcopiedBindId = OutId;
+		OutId = LoadIds.BindLoadId(*SaveIdIt++, MemcopiedBindId);
+	}
 }
 
-[[nodiscard]] static FLoadStructPlan MakeSchemaLoadPlan(const FStructSchema& From, const FSchemaBinding& To, TConstArrayView<FMemberId> ToMemberIds, TConstArrayView<FStructSchemaId> ToStructIds, SubsetByteArray& OutSubsetSchemas)
+[[nodiscard]] static FLoadStructPlan MakeSchemaLoadPlan(const FStructSchema& From, const FSchemaBinding& To, TConstArrayView<FMemberId> ToMemberIds, TConstArrayView<FStructSchemaId> ToStructIds, FLoadIdBinder& LoadIds, SubsetByteArray& OutSubsetSchemas)
 {
 	uint16 NumEnums = CountEnums(From);
 	if (From.NumMembers < To.NumMembers || NumEnums || HasDifferentSupers(From, To, ToStructIds))
 	{
-		CreateSubsetBindingWithoutEnumIds(From, To, ToMemberIds, NumEnums, OutSubsetSchemas);
+		CreateSubsetBindingWithoutEnumIds(From, To, ToMemberIds, NumEnums, LoadIds, OutSubsetSchemas);
 	}
 	else
 	{
@@ -228,7 +291,7 @@ static void CloneBindingWithReplacedStructIds(const FSchemaId* FromIds, const FS
 
 		if (From.NumInnerSchemas > 0)
 		{
-			CloneBindingWithReplacedStructIds(From.GetInnerSchemas(), To, OutSubsetSchemas);
+			CloneBindingWithReplacedStructIds(From.GetInnerSchemas(), To, LoadIds, OutSubsetSchemas);
 		}
 		// else reuse existing bindings
 	}
@@ -243,91 +306,159 @@ static void CloneBindingWithReplacedStructIds(const FSchemaId* FromIds, const FS
 	return NullOpt;
 }
 
-[[nodiscard]] static FLoadStructPlan MakeLoadPlan(const FStructSchema& From, const FSchemaBinding& To, TConstArrayView<FMemberId> ToMemberIds, TConstArrayView<FStructSchemaId> ToStructIds, SubsetByteArray& OutSubsetSchemas)
+[[nodiscard]] static FLoadStructPlan MakeLoadPlan(const FStructSchema& From, const FSchemaBinding& To, TConstArrayView<FMemberId> ToMemberIds, TConstArrayView<FStructSchemaId> ToStructIds, FLoadIdBinder& LoadIds, SubsetByteArray& OutSubsetSchemas)
 {
 	TOptional<FLoadStructMemcpy> Memcpy = TryMakeMemcpyPlan(From, To, ToMemberIds);
-	return Memcpy ? FLoadStructPlan(Memcpy.GetValue()) : MakeSchemaLoadPlan(From, To, ToMemberIds, ToStructIds, OutSubsetSchemas);
+	return Memcpy ? FLoadStructPlan(Memcpy.GetValue()) : MakeSchemaLoadPlan(From, To, ToMemberIds, ToStructIds, LoadIds, OutSubsetSchemas);
 }
 
-struct FMemberlessDummyBinding : ICustomBinding
+struct FMissingBinding : ICustomBinding
 {
-	virtual void SaveCustom(FMemberBuilder& Dst, const void* Src, const void* Default, const FSaveContext& Ctx) override { check(false); }
-	virtual void LoadCustom(void* Dst, FStructView Src, ECustomLoadMethod Method, const FLoadBatch& Batch) const override { check(false);}
-	virtual bool DiffCustom(const void* StructA, const void* StructB) const override { check(false); return false; }
+	FMissingBinding(const TCHAR* InType) : Type(InType) {}
+	const TCHAR* Type;
+	virtual void SaveCustom(FMemberBuilder& Dst, const void* Src, const void* Default, const FSaveContext& Ctx) override { checkf(false, TEXT("Can't save %s"), Type); }
+	virtual void LoadCustom(void* Dst, FStructView Src, ECustomLoadMethod Method, const FLoadBatch& Batch) const override { checkf(false, TEXT("Can't load %s"), Type); }
+	virtual bool DiffCustom(const void* StructA, const void* StructB) const override { checkf(false, TEXT("Can't diff %s"), Type); return false; }
 };
-static FMemberlessDummyBinding GMemberlessBinding;
+static FMissingBinding GMemberlessBinding(TEXT("Memberless struct binding"));
+static FMissingBinding GUnboundBinding{TEXT("Type-erased struct binding")};
 
-FLoadBatchPtr CreateLoadPlans(FReadBatchId ReadId, const FDeclarations& Declarations, const FCustomBindings& Customs, const FSchemaBindings& Schemas, TConstArrayView<FStructSchemaId> RuntimeIds)
+struct FLoadPlanner
 {
-	check(NumStructSchemas(ReadId) == RuntimeIds.Num());
-
-	// Temporary data structures
-	const uint32 NumPlans = RuntimeIds.Num();
-	TArray<FLoadStructPlan, TInlineAllocator<256>> Plans;
-	TArray<uint32, TInlineAllocator<256>> SubsetSchemaSizes;
-	SubsetByteArray SubsetSchemaData;
-	Plans.SetNumUninitialized(NumPlans);
-	SubsetSchemaSizes.SetNumUninitialized(NumPlans);
-
-	// Create plans
-	for (FStructSchemaId SavedId = { 0 }; SavedId.Idx < NumPlans; ++SavedId.Idx)
-	{
-		FStructSchemaId RuntimeId = RuntimeIds[SavedId.Idx];
-		int32 SubsetSchemaOffset = SubsetSchemaData.Num();
-		if (const ICustomBinding* Custom = Customs.FindStruct(RuntimeId))
+	// Constructor parameters
+	const FReadBatchId								ReadId;
+	const FDeclarations&							Declarations;
+	const FCustomBindings&							Customs;
+	const FSchemaBindings&							Schemas;
+	const TConstArrayView<FStructSchemaId>			RuntimeIds;
+
+	// Temporary data structures	
+	TArray<FLoadStructPlan, TInlineAllocator<256>>	Plans;
+	TArray<uint32, TInlineAllocator<256>>			SubsetSchemaSizes;
+	SubsetByteArray									SubsetSchemaData;
+	TSet<FStructSchemaId>							UnboundSaveIds;
+	
+	FLoadBatchPtr CreatePlans()
+	{
+		check(NumStructSchemas(ReadId) == RuntimeIds.Num());
+
+		// Make load plans for saved schemas
+		const uint32 NumPlans = RuntimeIds.Num();
+		Plans.SetNumUninitialized(NumPlans);
+		SubsetSchemaSizes.SetNumUninitialized(NumPlans);
+		FLoadIdBinder LoadIds(RuntimeIds);
+		for (uint32 Idx = 0; Idx < NumPlans; ++Idx)
 		{
-			Plans[SavedId.Idx] = FLoadStructPlan(*Custom);
+			FLoadIdMapping Mapping;
+			Mapping.SaveId = { Idx };
+			Mapping.LoadId = { Idx };
+			Mapping.BindId = RuntimeIds[Idx];
+			CreatePlan(Mapping, LoadIds);
 		}
-		else
+
+		// Make load plans for type-erased/ExplicitBindName structs needed by already created load plans
+		if (LoadIds.NumMappings() > 0)
 		{
-			const FStructSchema& From = ResolveStructSchema(ReadId, SavedId);
-			if (From.NumMembers)
+			Plans.Reserve(NumPlans + LoadIds.NumMappings());
+			SubsetSchemaSizes.Reserve(NumPlans + LoadIds.NumMappings());
+			for (int32 Idx = 0; Idx < LoadIds.NumMappings(); ++Idx)
 			{
-				const FSchemaBinding& To = Schemas.GetStruct(RuntimeId);
-				// Possible optimization - some simple memcpy cases doesn't need to resolve the declaration
-				TConstArrayView<FMemberId> ToMemberIds = Declarations.Get(RuntimeId).GetMemberOrder();
-				Plans[SavedId.Idx] = MakeLoadPlan(From, To, ToMemberIds, RuntimeIds, /* out */ SubsetSchemaData);	
+				check(LoadIds.GetMapping(Idx).LoadId.Idx == static_cast<uint32>(Plans.Num()));
+				Plans.AddUninitialized();
+				SubsetSchemaSizes.AddUninitialized();
+				CreatePlan(LoadIds.GetMapping(Idx), /* might grow */ LoadIds);
 			}
-			else
+
+			// Verify that all UnboundSaveIds were bound by some load plan
+			for (int32 Idx = 0; !UnboundSaveIds.IsEmpty() && Idx < LoadIds.NumMappings(); ++Idx)
 			{
-				checkf(!From.Type.Scope && From.Type.Name.NumParameters == 2, TEXT("Only range-bound template parameters are memberless. "
-					"They're always anonymous two-parameter types and uninstantiable as structs, bound via MakeAnonymousParametricType()"));
-				Plans[SavedId.Idx] = FLoadStructPlan(GMemberlessBinding);
+				UnboundSaveIds.Remove(LoadIds.GetMapping(Idx).SaveId);
 			}
 		}
-		
-		SubsetSchemaSizes[SavedId.Idx] = SubsetSchemaData.Num() - SubsetSchemaOffset;
+
+		for (FStructSchemaId UnboundSaveId : UnboundSaveIds)
+		{
+			checkf(false, TEXT("Unbound struct '%s' can't be loaded"), *Declarations.GetDebug().Print(RuntimeIds[UnboundSaveId.Idx]));
+		}
+
+		// Allocate load batch, copy plans and subset schemas, and fixup subset schema plans
+		return CreateBatch();
 	}
-	
-	// Allocate load batch, copy plans and subset schemas, and fixup subset schema plans
-	SIZE_T Bytes = sizeof(FLoadBatch) + sizeof(FLoadStructPlan) * Plans.Num() + SubsetSchemaData.Num();
-	FLoadBatch Header = { ReadId, NumPlans };
-	FLoadBatch* Out = new (FMemory::Malloc(Bytes)) FLoadBatch{Header};
-	FMemory::Memcpy(Out->Plans, Plans.GetData(), sizeof(FLoadStructPlan) * Plans.Num());
-	if (SubsetSchemaData.Num() > 0)
-	{
-		uint8* OutSubsetData = reinterpret_cast<uint8*>(Out->Plans + Plans.Num());
-		FMemory::Memcpy(OutSubsetData, SubsetSchemaData.GetData(), SubsetSchemaData.Num());
+private:
+
+	FLoadBatchPtr CreateBatch()
+	{
+		const uint32 NumPlans = Plans.Num();
+
+		SIZE_T Bytes = sizeof(FLoadBatch) + sizeof(FLoadStructPlan) * NumPlans + SubsetSchemaData.Num();
+		FLoadBatch Header = { ReadId, NumPlans };
+		FLoadBatch* Out = new (FMemory::Malloc(Bytes)) FLoadBatch{Header};
+		FMemory::Memcpy(Out->Plans, Plans.GetData(), sizeof(FLoadStructPlan) * NumPlans);
+		if (SubsetSchemaData.Num() > 0)
+		{
+			uint8* OutSubsetData = reinterpret_cast<uint8*>(Out->Plans + NumPlans);
+			FMemory::Memcpy(OutSubsetData, SubsetSchemaData.GetData(), SubsetSchemaData.Num());
 		
-		// Update plans with actual subset schema pointers
-		const uint8* It = OutSubsetData;
-		for (uint32 Idx = 0; Idx < NumPlans; ++Idx)
+			// Update plans with actual subset schema pointers
+			const uint8* It = OutSubsetData;
+			for (uint32 Idx = 0; Idx < NumPlans; ++Idx)
+			{
+				if (int32 Size = SubsetSchemaSizes[Idx])
+				{
+					check(IsAligned(Size, alignof(FSchemaBinding)));
+					check(Plans[Idx].IsSchema());
+					bool bSparse = Plans[Idx].IsSparseSchema();
+					Out->Plans[Idx] = FLoadStructPlan(*reinterpret_cast<const FSchemaBinding*>(It), ELeafWidth::B32, bSparse);
+					It += Size;
+				}
+			}
+			check(It == OutSubsetData + SubsetSchemaData.Num());
+		}
+
+		return FLoadBatchPtr(Out);
+	}
+
+	void CreatePlan(FLoadIdMapping Mapping, FLoadIdBinder& LoadIds)
+	{
+		int32 SubsetSchemaOffset = SubsetSchemaData.Num();	
+		Plans[Mapping.LoadId.Idx] = CreatePlanInner(Mapping, LoadIds);
+		SubsetSchemaSizes[Mapping.LoadId.Idx] = static_cast<uint32>(SubsetSchemaData.Num() - SubsetSchemaOffset);
+	}
+
+	FLoadStructPlan CreatePlanInner(FLoadIdMapping Mapping, FLoadIdBinder& LoadIds)
+	{
+		if (const ICustomBinding* Custom = Customs.FindStruct(Mapping.BindId))
 		{
-			if (int32 Size = SubsetSchemaSizes[Idx])
+			return FLoadStructPlan(*Custom);
+		}
+
+		const FStructSchema& From = ResolveStructSchema(ReadId, Mapping.SaveId);
+		if (From.NumMembers)
+		{
+			if (const FSchemaBinding* To = Schemas.FindStruct(Mapping.BindId))
 			{
-				check(IsAligned(Size, alignof(FSchemaBinding)));
-				check(Plans[Idx].IsSchema());
-				bool bSparse = Plans[Idx].IsSparseSchema();
-				Out->Plans[Idx] = FLoadStructPlan(*reinterpret_cast<const FSchemaBinding*>(It), ELeafWidth::B32, bSparse);
-				It += Size;
+				// Possible optimization - some simple memcpy cases doesn't need to resolve the declaration
+				TConstArrayView<FMemberId> ToMemberIds = Declarations.Get(To->DeclId).GetMemberOrder();
+				return MakeLoadPlan(From, *To, ToMemberIds, RuntimeIds, LoadIds, /* out */ SubsetSchemaData);	
 			}
+			
+			// Type-erased structs
+			UnboundSaveIds.Add(Mapping.SaveId);
+			return FLoadStructPlan(GUnboundBinding);
 		}
-		check(It == OutSubsetData + SubsetSchemaData.Num());
+		
+		checkf(!From.Type.Scope && From.Type.Name.NumParameters == 2, TEXT("Only range-bound template parameters are memberless. "
+				"They're always anonymous two-parameter types and uninstantiable as structs, bound via MakeAnonymousParametricType()"));
+		return FLoadStructPlan(GMemberlessBinding);
 	}
+};
 
-	return FLoadBatchPtr(Out);
+FLoadBatchPtr CreateLoadPlans(FReadBatchId ReadId, const FDeclarations& Declarations, const FCustomBindings& Customs, const FSchemaBindings& Schemas, TConstArrayView<FStructSchemaId> RuntimeIds)
+{
+	return FLoadPlanner{ReadId, Declarations, Customs, Schemas, RuntimeIds}.CreatePlans();
 }
 
+
 ////////////////////////////////////////////////////////////////////////////
 
 inline static void SetBit(uint8& Out, uint8 Idx, bool bValue)
diff --git a/Engine/Plugins/Experimental/PlainProps/Source/Private/PlainPropsSave.cpp b/Engine/Plugins/Experimental/PlainProps/Source/Private/PlainPropsSave.cpp
--- a/Engine/Plugins/Experimental/PlainProps/Source/Private/PlainPropsSave.cpp
+++ b/Engine/Plugins/Experimental/PlainProps/Source/Private/PlainPropsSave.cpp
@@ -310,27 +310,34 @@ static void SaveMember(FMemberBuilder& Out, const void* Struct, FMemberId Name,
 	Out.AddStruct(Name, Member.Id, SaveStruct(At(Struct, Member.Offset), Member.Id, Ctx));
 }
 
-FBuiltStruct* SaveStruct(const void* Struct, FStructSchemaId Id, const FSaveContext& Ctx)
+FBuiltStruct* SaveStruct(const void* Struct, FStructSchemaId BindId, const FSaveContext& Ctx)
 {
-	const FStructDeclaration& Declaration = Ctx.Declarations.Get(Id);
-
+	const FStructDeclaration* Declaration = nullptr;
 	FMemberBuilder Out;
-	if (ICustomBinding* Custom = Ctx.Customs.FindStruct(Id))
+	if (FCustomBindingEntry Custom = Ctx.Customs.FindStructToSave(BindId))
 	{
-		Custom->SaveCustom(Out, Struct, nullptr, Ctx);
+		Custom.Binding->SaveCustom(Out, Struct, nullptr, Ctx);
+		Declaration = &Ctx.Declarations.Get(Custom.DeclId);	
 	}
-	else for (FMemberVisitor It(Ctx.Schemas.GetStruct(Id)); It.HasMore(); )
+	else
 	{
-		FMemberId Name = Declaration.GetMemberOrder()[It.GetIndex()];
-		switch (It.PeekKind())
+		const FSchemaBinding& Schema = Ctx.Schemas.GetStruct(BindId);
+		Declaration = &Ctx.Declarations.Get(Schema.DeclId);	
+		TConstArrayView<FMemberId> MemberOrder = Declaration->GetMemberOrder();
+
+		for (FMemberVisitor It(Schema); It.HasMore(); )
 		{
-			case EMemberKind::Leaf:		SaveMember(Out, Struct, Name, Ctx, It.GrabLeaf());		break;
-			case EMemberKind::Range:	SaveMember(Out, Struct, Name, Ctx, It.GrabRange());		break;
-			case EMemberKind::Struct:	SaveMember(Out, Struct, Name, Ctx, It.GrabStruct());	break;
+			FMemberId Name = MemberOrder[It.GetIndex()];
+			switch (It.PeekKind())
+			{
+				case EMemberKind::Leaf:		SaveMember(Out, Struct, Name, Ctx, It.GrabLeaf());		break;
+				case EMemberKind::Range:	SaveMember(Out, Struct, Name, Ctx, It.GrabRange());		break;
+				case EMemberKind::Struct:	SaveMember(Out, Struct, Name, Ctx, It.GrabStruct());	break;
+			}
 		}
 	}
-	
-	return Out.BuildAndReset(Ctx.Scratch, Declaration, Ctx.Declarations.GetDebug());
+
+	return Out.BuildAndReset(Ctx.Scratch, *Declaration, Ctx.Declarations.GetDebug());
 }
 
 ////////////////////////////////////////////////////////////////////////////////////////////////
@@ -474,7 +481,7 @@ static bool DiffItem(const uint8* A, const uint8* B, const FSaveContext& Ctx, FR
 
 static bool DiffItem(const uint8* A, const uint8* B, const FSaveContext& Ctx, FStructSchemaId Id)
 {
-	if (ICustomBinding* Custom = Ctx.Customs.FindStruct(Id))
+	if (const ICustomBinding* Custom = Ctx.Customs.FindStruct(Id))
 	{
 		return Custom->DiffCustom(A, B);
 	}
@@ -524,30 +531,37 @@ static void SaveMemberDelta(FMemberBuilder& Out, const void* Struct, const void*
 	}
 }
 
-FBuiltStruct* SaveStructDelta(const void* Struct, const void* Default, FStructSchemaId Id, const FSaveContext& Ctx)
+FBuiltStruct* SaveStructDelta(const void* Struct, const void* Default, FStructSchemaId BindId, const FSaveContext& Ctx)
 {
-	const FStructDeclaration& Declaration = Ctx.Declarations.Get(Id);
-
+	const FStructDeclaration* Declaration = nullptr;
 	FMemberBuilder Out;
-	if (ICustomBinding* Custom = Ctx.Customs.FindStruct(Id))
+	if (FCustomBindingEntry Custom = Ctx.Customs.FindStructToSave(BindId))
 	{
-		if (Custom->DiffCustom(Struct, Default))
+		if (Custom.Binding->DiffCustom(Struct, Default))
 		{
-			Custom->SaveCustom(Out, Struct, Default, Ctx);
+			Declaration = &Ctx.Declarations.Get(Custom.DeclId);	
+			Custom.Binding->SaveCustom(Out, Struct, Default, Ctx);
 		}
 	}
-	else for (FMemberVisitor It(Ctx.Schemas.GetStruct(Id)); It.HasMore(); )
+	else
 	{
-		FMemberId Name = Declaration.GetMemberOrder()[It.GetIndex()];
-		switch (It.PeekKind())
+		const FSchemaBinding& Schema = Ctx.Schemas.GetStruct(BindId);
+		Declaration = &Ctx.Declarations.Get(Schema.DeclId);	
+		TConstArrayView<FMemberId> MemberOrder = Declaration->GetMemberOrder();
+
+		for (FMemberVisitor It(Schema); It.HasMore(); )
 		{
-			case EMemberKind::Leaf:		SaveMemberDelta(Out, Struct, Default, Name, Ctx, It.GrabLeaf());	break;
-			case EMemberKind::Range:	SaveMemberDelta(Out, Struct, Default, Name, Ctx, It.GrabRange());	break;
-			case EMemberKind::Struct:	SaveMemberDelta(Out, Struct, Default, Name, Ctx, It.GrabStruct());	break;
+			FMemberId Name = MemberOrder[It.GetIndex()];
+			switch (It.PeekKind())
+			{
+				case EMemberKind::Leaf:		SaveMemberDelta(Out, Struct, Default, Name, Ctx, It.GrabLeaf());	break;
+				case EMemberKind::Range:	SaveMemberDelta(Out, Struct, Default, Name, Ctx, It.GrabRange());	break;
+				case EMemberKind::Struct:	SaveMemberDelta(Out, Struct, Default, Name, Ctx, It.GrabStruct());	break;
+			}
 		}
 	}
 	
-	return Out.IsEmpty() ? nullptr : Out.BuildAndReset(Ctx.Scratch, Declaration, Ctx.Declarations.GetDebug());
+	return Out.IsEmpty() ? nullptr : Out.BuildAndReset(Ctx.Scratch, *Declaration, Ctx.Declarations.GetDebug());
 }
 
 } // namespace PlainProps
\ No newline at end of file
diff --git a/Engine/Plugins/Experimental/PlainProps/Source/Private/PlainPropsWrite.cpp b/Engine/Plugins/Experimental/PlainProps/Source/Private/PlainPropsWrite.cpp
--- a/Engine/Plugins/Experimental/PlainProps/Source/Private/PlainPropsWrite.cpp
+++ b/Engine/Plugins/Experimental/PlainProps/Source/Private/PlainPropsWrite.cpp
@@ -18,10 +18,12 @@ namespace PlainProps
 // Rewrite as a more compact data structure once we get a large number of ids
 struct FWriteIds
 {
-	FWriteIds(const FIdIndexerBase& DeclaredIds, const FBuiltSchemas& Schemas, ESchemaFormat Format);
+	FWriteIds(const FIdIndexerBase& DeclaredIds, const IStructBindIds& BindIds, const FBuiltSchemas& Schemas, ESchemaFormat Format);
 
 	bool								HasStableNames() const { return !Names.IsEmpty(); }
 
+	const IStructBindIds&				BindIds;
+
 	TArray<FOptionalNameId>				Names;
 	TArray<FOptionalNestedScopeId>		NestedScopes;
 	TArray<FOptionalParametricTypeId>	ParametricTypes;
@@ -43,7 +45,7 @@ struct FWriteIds
 	FParametricTypeId					Remap(FParametricTypeId Old) const		{ return ParametricTypes[Old.Idx].Get(); }
 	FTypenameId							Remap(FTypenameId Old) const			{ return Old.IsConcrete() ? FTypenameId(Remap(Old.AsConcrete())) : FTypenameId(Remap(Old.AsParametric())); }
 	FTypeId								Remap(FTypeId Old) const				{ return { Remap(Old.Scope), Remap(Old.Name) }; }
-	FStructSchemaId						RemapStruct(FSchemaId Old) const		{ return Structs[Old.Idx].Get(); }
+	FStructSchemaId						RemapStruct(FSchemaId Old) const;
 	FEnumSchemaId						RemapEnum(FSchemaId Old) const			{ return Enums[Old.Idx].Get(); }
 
 	template<typename T>
@@ -63,17 +65,17 @@ static TConstArrayView<FNameId> GetUsedNames(const FBuiltEnumSchema& Used)
 
 struct FUsedIds
 {
-	const FIdIndexerBase& Declared;
+	const FIdIndexerBase& Ids;
 	TBitArray<> Names;
 	TBitArray<> NestedScopes;
 	TBitArray<> ParametricTypes;
 
 
-	explicit FUsedIds(const FIdIndexerBase& InDeclared)
-	: Declared(InDeclared)
-	, Names(false, Declared.NumNames())
-	, NestedScopes(false, Declared.GetNestedScopes().Num())
-	, ParametricTypes(false, Declared.GetParametricTypes().Num())
+	explicit FUsedIds(const FIdIndexerBase& InIds)
+	: Ids(InIds)
+	, Names(false, Ids.NumNames())
+	, NestedScopes(false, Ids.GetNestedScopes().Num())
+	, ParametricTypes(false, Ids.GetParametricTypes().Num())
 	{}
 
 	template<class PartialSchemaType>
@@ -122,7 +124,7 @@ struct FUsedIds
 			{
 				Used = true;
 
-				FNestedScope Nested = Declared.Resolve(Scope.AsNested());
+				FNestedScope Nested = Ids.Resolve(Scope.AsNested());
 				MarkUsed(Nested.Outer);
 				MarkUsed(Nested.Inner.Name);
 			}
@@ -142,7 +144,7 @@ struct FUsedIds
 			{
 				Used = true;
 
-				FParametricTypeView ParametricType = Declared.Resolve(Typename.AsParametric());
+				FParametricTypeView ParametricType = Ids.Resolve(Typename.AsParametric());
 				MarkUsed(ParametricType.Name);
 				for (FTypeId Parameter : ParametricType.GetParameters())
 				{
@@ -189,10 +191,11 @@ static void CopyUsedIds(TArray<T>& Out, const TBitArray<>& Used, C&& Ids)
 	}
 }
 
-FWriteIds::FWriteIds(const FIdIndexerBase& Declared, const FBuiltSchemas& Schemas, ESchemaFormat Format)
+FWriteIds::FWriteIds(const FIdIndexerBase& Ids, const IStructBindIds& InBindIds, const FBuiltSchemas& Schemas, ESchemaFormat Format)
+: BindIds(InBindIds)
 {
-	Structs.Init(NoId, Declared.NumStructs());
-	Enums.Init(NoId, Declared.NumEnums());
+	Structs.Init(NoId, Ids.NumStructs());
+	Enums.Init(NoId, Ids.NumEnums());
 
 	// Generate new struct and enum schema indices
 	uint32 NewSchemaIdx = 0;
@@ -210,18 +213,18 @@ FWriteIds::FWriteIds(const FIdIndexerBase& Declared, const FBuiltSchemas& Schema
 	if (Format == ESchemaFormat::StableNames)
 	{
 		// Generate new name indices
-		FUsedIds Used(Declared);
+		FUsedIds Used(Ids);
 		Used.DetectUsage(Schemas.Structs);
 		Used.DetectUsage(Schemas.Enums);
 		MakeRemapping(Names, Used.Names);
 
 		// Generate new nested scope and parametric type mappings
 		KeptScopes.Reserve(MakeRemapping(/* out */ NestedScopes, Used.NestedScopes));
-		KeptParametrics.Reserve(MakeRemapping(/* out */ ParametricTypes, Used.ParametricTypes, Declared.GetParametricTypes()));
+		KeptParametrics.Reserve(MakeRemapping(/* out */ ParametricTypes, Used.ParametricTypes, Ids.GetParametricTypes()));
 
 		// Copy nested scopes and parametric types
-		CopyUsedIds(/* out */ KeptScopes, Used.NestedScopes, Declared.GetNestedScopes());
-		CopyUsedIds(/* out */ KeptParametrics, Used.ParametricTypes, Declared.GetParametricTypes().GetAllTypes());
+		CopyUsedIds(/* out */ KeptScopes, Used.NestedScopes, Ids.GetNestedScopes());
+		CopyUsedIds(/* out */ KeptParametrics, Used.ParametricTypes, Ids.GetParametricTypes().GetAllTypes());
 
 		// Remap copied nested scopes
 		for (FNestedScope& KeptScope : KeptScopes)
@@ -231,13 +234,13 @@ FWriteIds::FWriteIds(const FIdIndexerBase& Declared, const FBuiltSchemas& Schema
 		}
 
 		// Remap copied parametric types and copy parameters
-		TConstArrayView<FTypeId> DeclaredParameters = Declared.GetParametricTypes().GetAllParameters();
+		TConstArrayView<FTypeId> AllParameters = Ids.GetParametricTypes().GetAllParameters();
 		for (FParametricType& KeptType : KeptParametrics)
 		{
 			FParameterIndexRange OldParameters = KeptType.Parameters;
 			KeptType.Name = Remap(KeptType.Name);
 			KeptType.Parameters.Idx = KeptParameters.Num();
-			KeptParameters.Append(DeclaredParameters.Slice(OldParameters.Idx, OldParameters.NumParameters));
+			KeptParameters.Append(AllParameters.Slice(OldParameters.Idx, OldParameters.NumParameters));
 		}
 
 		// Remap copied parameters
@@ -248,6 +251,19 @@ FWriteIds::FWriteIds(const FIdIndexerBase& Declared, const FBuiltSchemas& Schema
 	}
 }
 
+FStructSchemaId
+FWriteIds::RemapStruct(FSchemaId OldBindId) const
+{
+	if (FOptionalStructSchemaId NewDeclId = Structs[OldBindId.Idx])
+	{
+		return NewDeclId.Get();
+	}
+
+	// Could optimize by caching Structs[OldBindId.Idx] here
+	FStructSchemaId OldDeclId = BindIds.GetDeclId(static_cast<FStructSchemaId>(OldBindId));
+	return Structs[OldDeclId.Idx].Get();
+}
+
 //////////////////////////////////////////////////////////////////////////
 
 static TArray<FMemberType> GetMemberTypes(const FBuiltStructSchema& Struct)
@@ -753,10 +769,10 @@ private:
 
 ////////////////////////////////////////////////////////////////////////////////////////////////////////////////////////////////////////////////////
 
-FWriter::FWriter(const FIdIndexerBase& DeclaredIds, const FBuiltSchemas& InSchemas, ESchemaFormat Format)
+FWriter::FWriter(const FIdIndexerBase& AllIds, const IStructBindIds& BindIds, const FBuiltSchemas& InSchemas, ESchemaFormat Format)
 : Schemas(InSchemas)
-, Debug(DeclaredIds)
-, NewIds(new FWriteIds(DeclaredIds, InSchemas, Format))
+, Debug(AllIds)
+, NewIds(new FWriteIds(AllIds, BindIds, InSchemas, Format))
 {}
 
 FWriter::~FWriter()
diff --git a/Engine/Plugins/Experimental/PlainProps/Source/Public/PlainPropsBind.h b/Engine/Plugins/Experimental/PlainProps/Source/Public/PlainPropsBind.h
--- a/Engine/Plugins/Experimental/PlainProps/Source/Public/PlainPropsBind.h
+++ b/Engine/Plugins/Experimental/PlainProps/Source/Public/PlainPropsBind.h
@@ -3,10 +3,12 @@
 #pragma once
 
 #include "Containers/ArrayView.h"
+#include "Containers/StringView.h"
 #include "Memory/MemoryFwd.h"
 #include "Memory/MemoryView.h"
 #include "PlainPropsCtti.h"
 #include "PlainPropsDeclare.h"
+#include "PlainPropsTypename.h"
 #include "PlainPropsRead.h"
 #include "PlainPropsTypes.h"
 #include <tuple>
@@ -30,6 +32,10 @@ template<class T> class TIdIndexer;
 
 ////////////////////////////////////////////////////////////////////////////////////////////////
 
+inline FAnsiStringView ToAnsiView(std::string_view Str) { return FAnsiStringView(Str.data(), Str.length()); }
+
+////////////////////////////////////////////////////////////////////////////////////////////////
+
 enum class ELeafBindType : uint8 { Bool, IntS, IntU, Float, Hex, Enum, Unicode, BitfieldBool };
 
 inline static constexpr ELeafBindType ToLeafBindType(ELeafType Type)
@@ -122,6 +128,7 @@ static_assert(sizeof(FMemberBindType) == 1);
 // Members are loaded in saved FStructSchema order, not current offset order unless upgrade layer reorders
 struct FSchemaBinding
 {
+	FStructSchemaId			DeclId;
 	uint16					NumMembers;
 	uint16					NumInnerSchemas;
 	uint16					NumInnerRanges;
@@ -234,6 +241,14 @@ protected: // for unit tests
 
 ////////////////////////////////////////////////////////////////////////////////////////////////
 
+struct FDualStructSchemaId
+{
+	FStructSchemaId Bind;
+	FStructSchemaId Decl;
+};
+
+////////////////////////////////////////////////////////////////////////////////////////////////
+
 enum ECustomLoadMethod { Construct, Assign };
 
 /// Load/save a struct with custom code to handle:
@@ -251,6 +266,15 @@ struct ICustomBinding
 	virtual bool				DiffCustom(const void* StructA, const void* StructB) const = 0;
 };
 
+struct FCustomBindingEntry
+{
+	FStructSchemaId BindId;
+	FStructSchemaId DeclId;
+	ICustomBinding* Binding = nullptr;
+
+	explicit operator bool() const { return !!Binding; }
+};
+
 class FCustomBindings
 {
 public:
@@ -258,22 +282,17 @@ public:
 	FCustomBindings(const FDebugIds& Dbg, const FCustomBindings* InBase = nullptr) : Base(InBase), Debug(Dbg) {}
 
 	// @param Binding must outlive this or call DropStruct()
-	PLAINPROPS_API void						BindStruct(FStructSchemaId Id, ICustomBinding& Binding);
-	ICustomBinding*							FindStruct(FStructSchemaId Id)								{ return Find(Id); }
-	const ICustomBinding*					FindStruct(FStructSchemaId Id) const						{ return Find(Id); }
-	PLAINPROPS_API void						DropStruct(FStructSchemaId Id);
+	PLAINPROPS_API void						BindStruct(FStructSchemaId BindId, FStructSchemaId DeclId, ICustomBinding& Binding);
+	const ICustomBinding*					FindStruct(FStructSchemaId BindId) const		{ return Find(BindId).Binding; }
+	FStructSchemaId							FindStructDeclId(FStructSchemaId BindId) const	{ return Find(BindId).DeclId; }
+	FCustomBindingEntry						FindStructToSave(FStructSchemaId BindId)		{ return Find(BindId); }
+	PLAINPROPS_API void						DropStruct(FStructSchemaId BindId);
 
 private:
-	PLAINPROPS_API ICustomBinding* Find(FStructSchemaId Id) const;
-
-	struct FEntry
-	{
-		FStructSchemaId	Id;
-		ICustomBinding* Binding;
-	};
-
+	PLAINPROPS_API FCustomBindingEntry		Find(FStructSchemaId BindId) const;
+	
 	const FCustomBindings*					Base = nullptr;
-	TArray<FEntry, TInlineAllocator<8>>		Entries;
+	TArray<FCustomBindingEntry, TInlineAllocator<8>>		Entries;
 	const FDebugIds&						Debug;
 };
 
@@ -529,95 +548,163 @@ struct FMemberBinding
 	TConstArrayView<FRangeBinding>	RangeBindings;		// Non-empty -> Range
 };
 
-class FSchemaBindings
+class FSchemaBindings : public IStructBindIds
 {
 public:
 	UE_NONCOPYABLE(FSchemaBindings);
 	explicit FSchemaBindings(const FDebugIds& In) : Debug(In) {}
 	PLAINPROPS_API ~FSchemaBindings();
 
-	PLAINPROPS_API void						BindStruct(FStructSchemaId Id, TConstArrayView<FMemberBinding> Schema);
-	PLAINPROPS_API const FSchemaBinding&	GetStruct(FStructSchemaId Id) const;
-	PLAINPROPS_API void						DropStruct(FStructSchemaId Id);
+	PLAINPROPS_API void						BindStruct(FStructSchemaId BindId, FStructSchemaId DeclId, TConstArrayView<FMemberBinding> Schema);
+	PLAINPROPS_API const FSchemaBinding*	FindStruct(FStructSchemaId BindId) const;
+	PLAINPROPS_API const FSchemaBinding&	GetStruct(FStructSchemaId BindId) const;
+	PLAINPROPS_API void						DropStruct(FStructSchemaId BindId);
 
 private:
 	TArray<TUniquePtr<FSchemaBinding>>		Bindings;
 	const FDebugIds&						Debug;
+
+	virtual FStructSchemaId GetDeclId(FStructSchemaId BindId) const override;
 };
 
 ////////////////////////////////////////////////////////////////////////////////////////////////
 
-template<class Ctti, typename Ids>
-FTypeId IndexStructOrEnumType();
+template<class Ids, ETypename Kind, typename Typename>
+FTypeId IndexBaseName()
+{	
+	if constexpr (Kind == ETypename::Bind && ExplicitBindName<Typename>)
+	{
+		return Ids::IndexNativeType(ToAnsiView(Typename::BindName));
+	}
+	else
+	{
+		return Ids::IndexNativeType(ToAnsiView(Typename::DeclName));	
+	}
+}
 
-template<typename Enum, typename Ids>
-FEnumSchemaId IndexEnum()
+template<class Ids, ETypename Kind, typename Typename>
+FTypeId IndexTypename()
 {
-	static FEnumSchemaId Id = Ids::IndexEnum(CttiOf<Enum>::Name);
-	return Id;
+	return IndexBaseName<Ids, Kind, Typename>();
 }
 
-template<typename Struct, typename Ids>
-FStructSchemaId IndexStruct()
+template<class Ids, ETypename Kind, ParametricName Typename>
+FTypeId IndexTypename()
 {
-	static FStructSchemaId Id = Ids::IndexStruct(IndexStructOrEnumType<CttiOf<Struct>, Ids>());
-	return Id;
+	FTypeId BaseName = IndexBaseName<Ids, Kind, Typename>();
+	return IndexParametricType<Ids, Kind>(BaseName, (typename Typename::Parameters*)nullptr);
 }
 
-template<typename Struct, typename Ids>
-FOptionalStructSchemaId IndexOptionalStruct()
+template<typename Struct, class Ids>
+FStructSchemaId IndexStructBindIdIfNeeded(FStructSchemaId DeclId)
 {
-	if constexpr (!std::is_void_v<Struct>)
+	using Typename = TTypename<Struct>;
+
+	if constexpr (ExplicitBindName<Typename> || ParametricName<Typename>)
 	{
-		return IndexStruct<Struct, Ids>();
+		return Ids::IndexStruct(IndexTypename<Ids, ETypename::Bind, Typename>());
+	}
+	else
+	{
+		return DeclId;
 	}
-	
-	return NoId;
 }
 
-template<typename T, typename Ids>
-FTypeId IndexAnyType()
+template<typename Struct, class Ids>
+FDualStructSchemaId IndexStructDualId()
 {
-	using RangeBinding = RangeBind<T>;
-	if constexpr (!std::is_void_v<RangeBinding>)
-	{
-		FTypeId InnerType = IndexAnyType<typename RangeBinding::ItemType, Ids>();
-		FTypeId SizeType = { NoId, Ids::IndexTypename(CttiOf<typename RangeBinding::SizeType>::Name) };
-		return Ids::GetIndexer().MakeAnonymousParametricType({InnerType, SizeType});
-	}
-	else 
+	using Typename = TTypename<Struct>;
+	
+	FDualStructSchemaId Out;
+	FTypeId DeclName = IndexTypename<Ids, ETypename::Decl, Typename>();
+	Out.Decl = Ids::IndexStruct(DeclName);
+	Out.Bind = Out.Decl;
+
+	if constexpr (ExplicitBindName<Typename> || ParametricName<Typename>)
 	{
-		return IndexStructOrEnumType<CttiOf<T>, Ids>();
+		FTypeId BindName = IndexTypename<Ids, ETypename::Bind, Typename>();
+		Out.Bind = BindName != DeclName ? Ids::IndexStruct(BindName) : Out.Decl;
 	}
+
+	return Out;
 }
 
-template<Arithmetic T, typename Ids>
-FTypeId IndexAnyType()
+// Cached by function static
+template<class Ids, typename Struct>
+FStructSchemaId GetStructDeclId()
 {
-	return { NoId, Ids::IndexTypename(CttiOf<T>::Name) };
+	static FStructSchemaId Id = Ids::IndexStruct(IndexTypename<Ids, ETypename::Decl, TTypename<Struct>>());
+	return Id;
 }
 
-template<typename Ids, typename... Ts>
-FTypeId IndexTemplatedType(FTypeId TemplatedType, const std::tuple<Ts...>*)
+// Cached by function static
+template<class Ids, typename Struct>
+FStructSchemaId GetStructBindId()
 {
-	FTypeId Parameters[] = { (IndexAnyType<Ts, Ids>())... };
-	return Ids::GetIndexer().MakeParametricType(TemplatedType, Parameters);
+	static FStructSchemaId Id = Ids::IndexStruct(IndexTypename<Ids, ETypename::Bind, TTypename<Struct>>());
+	return Id;
 }
 
-template<class Ctti, typename Ids>
-FTypeId IndexStructOrEnumType()
+// Cached by function static
+template<class Ids, typename Enum>
+FEnumSchemaId GetEnumId()
 {
-	FTypeId Type = Ids::IndexNativeType(Ctti::Name);
-	if constexpr (Templated<Ctti>)
+	static FEnumSchemaId Id = Ids::IndexEnum(CttiOf<Enum>::Name);
+	return Id;
+}
+
+template<class Ids, Arithmetic T>
+FTypeId IndexArithmeticName()
+{
+	static constexpr FUnpackedLeafType Leaf = ReflectArithmetic<T>;
+	return { NoId, Ids::IndexTypename(ToAnsiView(ArithmeticName<Leaf.Type, Leaf.Width>)) };
+}
+
+template<class Ids, ETypename, Arithmetic Leaf>
+FTypeId IndexParameterName()
+{
+	return IndexArithmeticName<Ids, Leaf>();
+}
+
+template<class Ids, ETypename, Enumeration Leaf>
+FTypeId IndexParameterName()
+{
+	return Ids::IndexEnum(CttiOf<Leaf>::Name);
+}
+
+template<class Ids, ETypename Kind, typename T>
+FTypeId IndexParameterName()
+{
+	using RangeBinding = RangeBind<T>;
+	if constexpr (std::is_void_v<RangeBinding>)
 	{
-		return IndexTemplatedType<Ids>(Type, (typename Ctti::TemplateArgs*)nullptr);
+		return IndexTypename<Ids, Kind, TTypename<T>>();
 	}
 	else
 	{
-		return Type;
+		FTypeId ItemTypename = IndexParameterName<Ids, Kind, typename RangeBinding::ItemType>();
+		FTypeId SizeTypename = IndexArithmeticName<Ids, typename RangeBinding::SizeType>();
+
+		if constexpr (Kind == ETypename::Decl)
+		{
+			// Type-erase range type
+			return Ids::GetIndexer().MakeAnonymousParametricType({ItemTypename, SizeTypename});
+		}
+		else
+		{
+			FTypeId RangeBindName = Ids::IndexNativeType(ToAnsiView(TTypename<T>::RangeBindName));
+			return Ids::GetIndexer().MakeParametricType(RangeBindName, {ItemTypename, SizeTypename});
+		}
 	}
 }
 
+template<class Ids, ETypename Kind, typename... Ts>
+FTypeId IndexParametricType(FTypeId TemplatedType, const std::tuple<Ts...>*)
+{
+	FTypeId Parameters[] = { (IndexParameterName<Ids, Kind, Ts>())... };
+	return Ids::GetIndexer().MakeParametricType(TemplatedType, Parameters);
+}
+
 ////////////////////////////////////////////////////////////////////////////////////////////////
 
 template<typename CustomBinding, class Runtime>
@@ -630,17 +717,19 @@ FStructSchemaId BindCustomStructOnce()
 
 		FBinding()
 		{
-			FTypeId Name = IndexStructOrEnumType<CttiOf<Type>, Ids>();
-			Id = Ids::IndexStruct(Name);
+			FTypeId DeclType = IndexTypename<Ids, ETypename::Decl, TTypename<Type>>();
+			DeclId = Ids::IndexStruct(DeclType);
+			BindId = IndexStructBindIdIfNeeded<Type, Ids>(DeclId);
+			
 			CustomBinding::template InitIds<Ids>();
-			Runtime::GetTypes().DeclareStruct(Id, Name, CustomBinding::MemberIds, CustomBinding::Occupancy);
-			Runtime::GetCustoms().BindStruct(Id, *this);
+			Runtime::GetTypes().DeclareStruct(DeclId, DeclType, CustomBinding::MemberIds, CustomBinding::Occupancy);
+			Runtime::GetCustoms().BindStruct(BindId, DeclId, *this);
 		}
 
 		~FBinding()
 		{
-			Runtime::GetCustoms().DropStruct(Id);
-			Runtime::GetTypes().DropStruct(Id);
+			Runtime::GetCustoms().DropStruct(BindId);
+			Runtime::GetTypes().DropStructRef(DeclId);
 		}
 
 		virtual void SaveCustom(FMemberBuilder& Dst, const void* Src, const void* Default, const FSaveContext& Ctx) override
@@ -657,20 +746,21 @@ FStructSchemaId BindCustomStructOnce()
 		{
 			return CustomBinding::Diff(*static_cast<const Type*>(A), *static_cast<const Type*>(B));
 		}
-
-		FStructSchemaId Id;
+		
+		FStructSchemaId DeclId;
+		FStructSchemaId BindId;
 	};
 
 	static FBinding Binding;
-	return Binding.Id;
+	return Binding.BindId;
 }
 
-template<class Type, class CustomBinding, class Runtime>
+template<class Struct, class CustomBinding, class Runtime>
 FMemberBindType BindMemberStruct(FOptionalSchemaId& OutSchema)
 {
 	if constexpr (std::is_void_v<CustomBinding>)
 	{
-		OutSchema = FOptionalSchemaId(IndexStruct<Type, typename Runtime::Ids>());
+		OutSchema = FOptionalSchemaId(GetStructBindId<typename Runtime::Ids, Struct>());
 	}
 	else
 	{
@@ -680,10 +770,10 @@ FMemberBindType BindMemberStruct(FOptionalSchemaId& OutSchema)
 	return FMemberBindType(FStructType{EMemberKind::Struct, /* IsDynamic */ 0, /* IsSuper */ 0});
 }
 
-template<typename Type, class Ids>
+template<typename Struct, class Ids>
 FMemberBindType BindInnermostType(FOptionalSchemaId& OutSchema)
 {
-	OutSchema = FOptionalSchemaId(IndexStruct<Type, Ids>());
+	OutSchema = FOptionalSchemaId(GetStructBindId<Ids, Struct>());
 	return FMemberBindType(FStructType{EMemberKind::Struct, /* IsDynamic */ 0, /* IsSuper */ 0});
 }
 
@@ -694,11 +784,11 @@ FMemberBindType BindInnermostType(FOptionalSchemaId& OutSchema)
 	return FMemberBindType(ReflectArithmetic<Type>);
 }
 
-template<Enumeration Type, class Ids>
+template<Enumeration Enum, class Ids>
 FMemberBindType BindInnermostType(FOptionalSchemaId& OutSchema)
 {
-	OutSchema = ToOptional(static_cast<FSchemaId>(IndexEnum<Type, Ids>()));
-	return FMemberBindType(ReflectEnum<Type>);
+	OutSchema = ToOptional(static_cast<FSchemaId>(GetEnumId<Ids, Enum>()));
+	return FMemberBindType(ReflectEnum<Enum>);
 }
 
 template<typename RangeBinding>
@@ -829,9 +919,17 @@ FEnumSchemaId DeclareNativeEnum(FDeclarations& Out, EEnumMode Mode)
 template<class Ctti, class Ids>
 FStructSchemaId DeclareNativeStruct(FDeclarations& Out, EMemberPresence Occupancy)
 {
-	FTypeId Type = IndexStructOrEnumType<Ctti, Ids>();
+	using Typename = TTypename<typename Ctti::Type>;
+	using SuperType = typename Ctti::Super;
+
+	FTypeId Type = IndexTypename<Ids, ETypename::Decl, Typename>();
 	FStructSchemaId Id = Ids::IndexStruct(Type);
-	FOptionalStructSchemaId SuperId = IndexOptionalStruct<typename Ctti::Super, Ids>();
+	FOptionalStructSchemaId SuperId;
+	if constexpr (!std::is_void_v<SuperType>)
+	{
+		SuperId = GetStructDeclId<Ids, SuperType>();
+	}
+
 	FMemberId MemberIds[Ctti::NumVars];
 	ForEachVar<Ctti>([&]<class Var>()
 	{ 
@@ -843,14 +941,14 @@ FStructSchemaId DeclareNativeStruct(FDeclarations& Out, EMemberPresence Occupanc
 }
 
 template<class Ctti, class Runtime>
-void BindNativeStruct(FSchemaBindings& Out, FStructSchemaId DeclaredId)
+void BindNativeStruct(FSchemaBindings& Out, FStructSchemaId BindId, FStructSchemaId DeclId)
 {
 	FMemberBinding MemberBindings[Ctti::NumVars];
 	ForEachVar<Ctti>([&]<class Var>()
 	{ 
 		MemberBindings[Var::Index] = BindMember<typename Var::Type, Runtime>(Var::Offset);
 	});
-	Out.BindStruct(DeclaredId, MemberBindings);
+	Out.BindStruct(BindId, DeclId, MemberBindings);
 }
 
 ////////////////////////////////////////////////////////////////////////////////////////////////
diff --git a/Engine/Plugins/Experimental/PlainProps/Source/Public/PlainPropsBuildSchema.h b/Engine/Plugins/Experimental/PlainProps/Source/Public/PlainPropsBuildSchema.h
--- a/Engine/Plugins/Experimental/PlainProps/Source/Public/PlainPropsBuildSchema.h
+++ b/Engine/Plugins/Experimental/PlainProps/Source/Public/PlainPropsBuildSchema.h
@@ -46,13 +46,13 @@ public:
 	using FStructDeclarations = TConstArrayView<TUniquePtr<FStructDeclaration>>;
 	using FEnumDeclarations = TConstArrayView<TUniquePtr<FEnumDeclaration>>;
 
-	PLAINPROPS_API explicit FSchemasBuilder(const FDeclarations& Declarations, FScratchAllocator& Scratch);
-	PLAINPROPS_API FSchemasBuilder(FStructDeclarations InStructs, FEnumDeclarations InEnums, const FDebugIds& InDebug, FScratchAllocator& Scratch);
+	PLAINPROPS_API FSchemasBuilder(const FDeclarations& Declarations, const IStructBindIds& BindIds, FScratchAllocator& Scratch);
+	PLAINPROPS_API FSchemasBuilder(FStructDeclarations InStructs, FEnumDeclarations InEnums, const IStructBindIds& BindIds, const FDebugIds& InDebug, FScratchAllocator& Scratch);
 	PLAINPROPS_API ~FSchemasBuilder();
 
 	PLAINPROPS_API FEnumSchemaBuilder&			NoteEnum(FEnumSchemaId Id);
-	PLAINPROPS_API FStructSchemaBuilder&		NoteStruct(FStructSchemaId Id);
-	PLAINPROPS_API void							NoteStructAndMembers(FStructSchemaId Id, const FBuiltStruct& Struct);
+	PLAINPROPS_API FStructSchemaBuilder&		NoteStruct(FStructSchemaId BindId);
+	PLAINPROPS_API void							NoteStructAndMembers(FStructSchemaId BindId, const FBuiltStruct& Struct);
 	PLAINPROPS_API FBuiltSchemas				Build();
 	
 	FScratchAllocator&							GetScratch() const { return Scratch; }
@@ -63,6 +63,7 @@ private:
 	FEnumDeclarations							DeclaredEnums;
 	TArray<int32>								StructIndices;
 	TArray<int32>								EnumIndices;
+	const IStructBindIds&						BindIds;
 	TPagedArray<FStructSchemaBuilder, 4096>		Structs;		// TPagedArray for stable references
 	TPagedArray<FEnumSchemaBuilder, 4096>		Enums;			// TPagedArray for stable references
 	FScratchAllocator&							Scratch;
diff --git a/Engine/Plugins/Experimental/PlainProps/Source/Public/PlainPropsCtti.h b/Engine/Plugins/Experimental/PlainProps/Source/Public/PlainPropsCtti.h
--- a/Engine/Plugins/Experimental/PlainProps/Source/Public/PlainPropsCtti.h
+++ b/Engine/Plugins/Experimental/PlainProps/Source/Public/PlainPropsCtti.h
@@ -3,30 +3,11 @@
 #pragma once
 
 #include "HAL/Platform.h"
+#include <tuple>
 
 namespace PlainProps 
 {
 
-#define _PP_ARITHMETIC_CTTI(T, N) struct T##_Ctti  { inline static constexpr char Name[] = #N; using Type = T; }; T##_Ctti CttiOfPtr(T*);
-_PP_ARITHMETIC_CTTI(float, f32)
-_PP_ARITHMETIC_CTTI(double, f64)
-_PP_ARITHMETIC_CTTI(bool, bool)
-_PP_ARITHMETIC_CTTI(int8, s8)
-_PP_ARITHMETIC_CTTI(uint8, u8)
-_PP_ARITHMETIC_CTTI(int16, s16)
-_PP_ARITHMETIC_CTTI(uint16, u16)
-_PP_ARITHMETIC_CTTI(int32, s32)
-_PP_ARITHMETIC_CTTI(uint32, u32)
-_PP_ARITHMETIC_CTTI(int64, s64)
-_PP_ARITHMETIC_CTTI(uint64, u64)
-_PP_ARITHMETIC_CTTI(char, utf8)
-_PP_ARITHMETIC_CTTI(char8_t, utf8)
-_PP_ARITHMETIC_CTTI(char16_t, utf16)
-_PP_ARITHMETIC_CTTI(char32_t, utf32)
-#undef _PP_ARITHMETIC_CTTI
-
-//////////////////////////////////////////////////////////////////////////
-
 template<typename T> struct TCttiOf { using Type = decltype(CttiOfPtr((T*)nullptr)); };
 template<typename T> using CttiOf = typename TCttiOf<T>::Type;
 
@@ -95,5 +76,4 @@ static constexpr void ForEachVar(Fn&& Visitor)
 template<class Ctti>
 concept Templated = requires { typename Ctti::TemplateArgs; };
 
-} // namespace PlainProps
-
+} // namespace PlainProps
\ No newline at end of file
diff --git a/Engine/Plugins/Experimental/PlainProps/Source/Public/PlainPropsDeclare.h b/Engine/Plugins/Experimental/PlainProps/Source/Public/PlainPropsDeclare.h
--- a/Engine/Plugins/Experimental/PlainProps/Source/Public/PlainPropsDeclare.h
+++ b/Engine/Plugins/Experimental/PlainProps/Source/Public/PlainPropsDeclare.h
@@ -9,6 +9,7 @@
 
 namespace PlainProps
 {
+
 struct FEnumerator
 {
 	FNameId					Name;
@@ -35,6 +36,7 @@ enum class EMemberPresence { RequireAll, AllowSparse };
 
 struct FStructDeclaration
 {
+	uint32					RefCount;
 	FStructSchemaId			Id;				// Could be removed, might allow declaration dedup among templated types
 	FTypeId					Type;			// Could be removed, might allow declaration dedup among templated types
 	FOptionalStructSchemaId	Super;
@@ -54,14 +56,15 @@ public:
 	UE_NONCOPYABLE(FDeclarations);
 	explicit FDeclarations(const FDebugIds& In) : Debug(In) {}
 
-	PLAINPROPS_API void								DeclareEnum(FEnumSchemaId Id, FTypeId Type, EEnumMode Mode, ELeafWidth Width, TConstArrayView<FEnumerator> Enumerators);
-	PLAINPROPS_API void								DeclareStruct(FStructSchemaId Id, FTypeId Type, TConstArrayView<FMemberId> MemberOrder, EMemberPresence Occupancy, FOptionalStructSchemaId Super = {});
+	PLAINPROPS_API void								DeclareEnum(FEnumSchemaId DeclId, FTypeId Type, EEnumMode Mode, ELeafWidth Width, TConstArrayView<FEnumerator> Enumerators);
+	// Declare struct with ref count 1 or increment it and check that previous declaration matches
+	PLAINPROPS_API void								DeclareStruct(FStructSchemaId DeclId, FTypeId Type, TConstArrayView<FMemberId> MemberOrder, EMemberPresence Occupancy, FOptionalStructSchemaId Super = {});
 	
-	void											DropEnum(FEnumSchemaId Id)		{ Check(Id); DeclaredEnums[Id.Idx].Reset(); }
-	void											DropStruct(FStructSchemaId Id)	{ Check(Id); DeclaredStructs[Id.Idx].Reset(); }
+	void											DropEnum(FEnumSchemaId DeclId)			{ Check(DeclId); DeclaredEnums[DeclId.Idx].Reset(); }
+	void											DropStructRef(FStructSchemaId DeclId);
 
-	const FEnumDeclaration&							Get(FEnumSchemaId Id) const		{ Check(Id); return *DeclaredEnums[Id.Idx]; }
-	const FStructDeclaration&						Get(FStructSchemaId Id) const	{ Check(Id); return *DeclaredStructs[Id.Idx]; }
+	const FEnumDeclaration&							Get(FEnumSchemaId DeclId) const			{ Check(DeclId); return *DeclaredEnums[DeclId.Idx]; }
+	const FStructDeclaration&						Get(FStructSchemaId DeclId) const		{ Check(DeclId); return *DeclaredStructs[DeclId.Idx]; }
 	
 	TConstArrayView<TUniquePtr<FEnumDeclaration>>	GetEnums() const				{ return DeclaredEnums; }
 	TConstArrayView<TUniquePtr<FStructDeclaration>>	GetStructs() const				{ return DeclaredStructs; }
@@ -80,4 +83,10 @@ protected:
 #endif
 };
 
+struct IStructBindIds
+{
+	virtual ~IStructBindIds() = default;
+	virtual FStructSchemaId GetDeclId(FStructSchemaId BindId) const = 0;
+};
+
 } // namespace PlainProps
\ No newline at end of file
diff --git a/Engine/Plugins/Experimental/PlainProps/Source/Public/PlainPropsSave.h b/Engine/Plugins/Experimental/PlainProps/Source/Public/PlainPropsSave.h
--- a/Engine/Plugins/Experimental/PlainProps/Source/Public/PlainPropsSave.h
+++ b/Engine/Plugins/Experimental/PlainProps/Source/Public/PlainPropsSave.h
@@ -33,7 +33,7 @@ FSaveContext MakeSaveContext(FDefaultStructs Defaults, FScratchAllocator& Scratc
 	return { Runtime::GetTypes(), Runtime::GetSchemas(), Runtime::GetCustoms(), Scratch, Defaults };
 }
 
-[[nodiscard]] PLAINPROPS_API FBuiltStructPtr SaveStruct(const void* Struct, FStructSchemaId Id, const FSaveContext& Context);
-[[nodiscard]] PLAINPROPS_API FBuiltStructPtr SaveStructDelta(const void* Struct, const void* Default, FStructSchemaId Id, const FSaveContext& Context);
+[[nodiscard]] PLAINPROPS_API FBuiltStructPtr SaveStruct(const void* Struct, FStructSchemaId BindId, const FSaveContext& Context);
+[[nodiscard]] PLAINPROPS_API FBuiltStructPtr SaveStructDelta(const void* Struct, const void* Default, FStructSchemaId BindId, const FSaveContext& Context);
 
 } // namespace PlainProps
\ No newline at end of file
diff --git a/Engine/Plugins/Experimental/PlainProps/Source/Public/PlainPropsStringUtil.h b/Engine/Plugins/Experimental/PlainProps/Source/Public/PlainPropsStringUtil.h
--- a/Engine/Plugins/Experimental/PlainProps/Source/Public/PlainPropsStringUtil.h
+++ b/Engine/Plugins/Experimental/PlainProps/Source/Public/PlainPropsStringUtil.h
@@ -34,9 +34,6 @@ struct TConcat
     inline static constexpr std::string_view Value{Chars.data(), Len};
 };
 
-template <const std::string_view&... Strs>
-inline constexpr std::string_view Concat = TConcat<Strs...>::Value;
-
 template<uint64 N>
 struct THexString
 {
@@ -59,7 +56,15 @@ struct THexString
 	inline static constexpr std::string_view Value{Chars + 16 - Len, Len};
 };
 
+} // namespace PlainProps::Private
+
+namespace PlainProps
+{
+
+template <const std::string_view&... Strs>
+inline constexpr std::string_view Concat = Private::TConcat<Strs...>::Value;
+
 template <uint64 N>
-inline constexpr std::string_view HexString = THexString<N>::Value;
+inline constexpr std::string_view HexString = Private::THexString<N>::Value;
 
-} // namespace PlainProps::Private
\ No newline at end of file
+} // namespace PlainProps
\ No newline at end of file
diff --git a/Engine/Plugins/Experimental/PlainProps/Source/Public/PlainPropsTypename.h b/Engine/Plugins/Experimental/PlainProps/Source/Public/PlainPropsTypename.h
new file mode 100644
--- /dev/null
+++ b/Engine/Plugins/Experimental/PlainProps/Source/Public/PlainPropsTypename.h
@@ -0,0 +1,95 @@
+// Copyright Epic Games, Inc. All Rights Reserved.
+
+#pragma once
+
+#include "PlainPropsCtti.h"
+#include "PlainPropsTypes.h"
+#include <string_view>
+
+namespace PlainProps 
+{
+
+// A declared struct can be bound to multiple runtime implementations. ETypename indicates which typename is intended.
+//
+// This enum is also used to type-erase ranges, e.g. in UE FString, TArray<char8_t>, TArray<char, TInlineAllocator<64>> 
+// all map to a UTF8 range and the runtime can change data structures without impacting saved format.
+//
+// This can be used for both schema bindings and custom bindings. E.g. weak and strong reference class can be saved identically
+// if the strong / weak semantics is considered a runtime detail. Typed references can be type-erased into an untyped reference.
+//
+// When there is only a single runtime memory representation the Bind name and Decl name is usually the same
+enum class ETypename
+{
+	Decl,	// Declared type name, saved to disk / schemas
+	Bind	// Bound type name, in-memory id that uniquely identifies a binding. Often same as Decl.
+};
+
+// Trait that defines DeclName, BindName and template Parameters that's part of the typename
+//
+// Also defines RangeBindName, used when otherwise nameless ranges are captured in Parameters, 
+// e.g. UE TPair<int, FString> need a different bind name than TPair<int, TArray<char8_t>>
+//
+// At least one of DeclName, BindName or RangeBindName should exist and be a constexpr std::string_view
+// Template parameters are captured as using Parameters = std::tuple<T>;
+template<typename T>
+struct TTypename
+{
+	inline static constexpr std::string_view DeclName = CttiOf<T>::Name;
+};
+
+template<class Typename>
+concept ParametricName = requires { typename Typename::Parameters; };
+
+template<class Typename>
+concept ExplicitBindName = requires { Typename::BindName; };
+
+////////////////////////////////////////////////////////////////////////////////////////////////
+
+// WIP type trait used to type-erase range bind names, might get folded into TTypename
+template<typename T>
+struct TShortTypename;
+
+template<typename T>
+constexpr std::string_view ShortTypename = TShortTypename<T>::Value;
+
+struct FOmitTypename { inline static constexpr std::string_view Value{}; };
+
+////////////////////////////////////////////////////////////////////////////////////////////////
+
+template<ELeafType Type, ELeafWidth Width>
+constexpr std::string_view IllegalArithmetic()
+{
+	static_assert((int)Type + (int)Width == -1, "Illegal ELeafType/ELeafWidth combination");
+	return "ERR";
+}
+
+template<ELeafType Type, ELeafWidth Width>
+constexpr std::string_view ArithmeticName = IllegalArithmetic<Type, Width>();
+
+template<> inline constexpr std::string_view ArithmeticName<ELeafType::Bool, ELeafWidth::B8>		= "bool";
+template<> inline constexpr std::string_view ArithmeticName<ELeafType::IntS, ELeafWidth::B8>		= "i8";
+template<> inline constexpr std::string_view ArithmeticName<ELeafType::IntS, ELeafWidth::B16>		= "i16";
+template<> inline constexpr std::string_view ArithmeticName<ELeafType::IntS, ELeafWidth::B32>		= "i32";
+template<> inline constexpr std::string_view ArithmeticName<ELeafType::IntS, ELeafWidth::B64>		= "i64";
+template<> inline constexpr std::string_view ArithmeticName<ELeafType::IntU, ELeafWidth::B8>		= "u8";
+template<> inline constexpr std::string_view ArithmeticName<ELeafType::IntU, ELeafWidth::B16>		= "u16";
+template<> inline constexpr std::string_view ArithmeticName<ELeafType::IntU, ELeafWidth::B32>		= "u32";
+template<> inline constexpr std::string_view ArithmeticName<ELeafType::IntU, ELeafWidth::B64>		= "u64";
+template<> inline constexpr std::string_view ArithmeticName<ELeafType::Float, ELeafWidth::B32>		= "f32";
+template<> inline constexpr std::string_view ArithmeticName<ELeafType::Float, ELeafWidth::B64>		= "f64";
+template<> inline constexpr std::string_view ArithmeticName<ELeafType::Hex, ELeafWidth::B8>			= "hex8";
+template<> inline constexpr std::string_view ArithmeticName<ELeafType::Hex, ELeafWidth::B16>		= "hex16";
+template<> inline constexpr std::string_view ArithmeticName<ELeafType::Hex, ELeafWidth::B32>		= "hex32";
+template<> inline constexpr std::string_view ArithmeticName<ELeafType::Hex, ELeafWidth::B64>		= "hex64";
+template<> inline constexpr std::string_view ArithmeticName<ELeafType::Unicode, ELeafWidth::B8>		= "utf8";
+template<> inline constexpr std::string_view ArithmeticName<ELeafType::Unicode, ELeafWidth::B16>	= "utf16";
+template<> inline constexpr std::string_view ArithmeticName<ELeafType::Unicode, ELeafWidth::B32>	= "utf32";
+
+template<Arithmetic T>
+struct TTypename<T>
+{
+	inline static constexpr FUnpackedLeafType Leaf = ReflectArithmetic<T>;
+	inline static constexpr std::string_view DeclName = ArithmeticName<Leaf.Type, Leaf.Width>;
+};
+
+} // namespace PlainProps
\ No newline at end of file
diff --git a/Engine/Plugins/Experimental/PlainProps/Source/Public/PlainPropsUeCoreBindings.h b/Engine/Plugins/Experimental/PlainProps/Source/Public/PlainPropsUeCoreBindings.h
--- a/Engine/Plugins/Experimental/PlainProps/Source/Public/PlainPropsUeCoreBindings.h
+++ b/Engine/Plugins/Experimental/PlainProps/Source/Public/PlainPropsUeCoreBindings.h
@@ -11,11 +11,10 @@
 #include "PlainPropsLoad.h"
 #include "PlainPropsRead.h"
 #include "PlainPropsIndex.h"
+#include "PlainPropsStringUtil.h"
 #include "UObject/NameTypes.h"
 
 
-PP_NAME_STRUCT(, FName);
-PP_NAME_STRUCT_TEMPLATE(, TSet);
 PP_REFLECT_STRUCT_TEMPLATE(, TTuple, void, Key, Value); // Todo handle TTuple and higher arities
 
 namespace UE::Math
@@ -23,7 +22,6 @@ namespace UE::Math
 PP_REFLECT_STRUCT(, FVector, void, X, Y, Z);
 PP_REFLECT_STRUCT(, FVector4, void, X, Y, Z, W);
 PP_REFLECT_STRUCT(, FQuat, void, X, Y, Z, W);
-PP_NAME_STRUCT(, FTransform);
 }
 
 namespace PlainProps::UE
@@ -431,15 +429,15 @@ struct FTransformBinding : ICustomBinding
 	FStructSchemaId VectorId;
 	FStructSchemaId QuatId;
 
-	template<typename Ids>
+	template<class Ids>
 	void InitIds(/*const FDeclarations& Declared*/) 
 	{
 		MemberIds[(uint8)EMember::Translate] = Ids::IndexMember("Translate");
 		MemberIds[(uint8)EMember::Rotate] = Ids::IndexMember("Rotate");
 		MemberIds[(uint8)EMember::Scale] = Ids::IndexMember("Scale");
 
-		VectorId = IndexStruct<FVector, Ids>();
-		QuatId = IndexStruct<FQuat, Ids>();
+		VectorId = GetStructDeclId<Ids, FVector>();
+		QuatId = GetStructDeclId<Ids, FQuat>();
 	}
 
 	PLAINPROPS_API void	Save(FMemberBuilder& Dst, const FTransform& Src, const FTransform* Default, const FSaveContext& Context) const;
@@ -465,6 +463,7 @@ struct FSetDeltaOps
 	}
 };
 
+
 template <typename T, typename KeyFuncs, typename SetAllocator>
 struct TSetDeltaBinding : ICustomBinding, FSetDeltaOps
 {
@@ -678,11 +677,77 @@ struct TSetDeltaBinding : ICustomBinding, FSetDeltaOps
 //	}
 //};
 
-}
+} // namespace PlainProps::UE
 
 namespace PlainProps
 {
 
+template <>
+struct TTypename<FName>
+{
+	inline static constexpr std::string_view DeclName = "Name";
+};
+
+template <>
+struct TTypename<FTransform>
+{
+	inline static constexpr std::string_view DeclName = "Transform";
+};
+
+template <typename K, typename V>
+struct TTypename<TPair<K,V>>
+{
+	inline static constexpr std::string_view DeclName = "Pair";
+	using Parameters = std::tuple<K, V>;
+};
+
+template <>
+struct TTypename<FString>
+{
+	inline static constexpr std::string_view RangeBindName = "String";
+};
+
+inline static constexpr std::string_view UeArrayName = "Array";
+inline static constexpr std::string_view UeSetName = "Set";
+inline static constexpr std::string_view UeMapName = "Map";
+
+template<typename T, typename Allocator>
+struct TTypename<TArray<T, Allocator>>
+{
+	inline static constexpr std::string_view RangeBindName = Concat<UeArrayName, ShortTypename<Allocator>>;
+};
+
+template<>
+struct TShortTypename<FDefaultAllocator> : FOmitTypename {};
+template<>
+struct TShortTypename<FDefaultSetAllocator> : FOmitTypename {};
+template<typename T>
+struct TShortTypename<DefaultKeyFuncs<T, false>> : FOmitTypename {};
+template<typename K, typename V>
+struct TShortTypename<TDefaultMapHashableKeyFuncs<K, V, false>> : FOmitTypename {};
+
+inline constexpr std::string_view InlineAllocatorPrefix = "InlX";
+template<int N>
+struct TShortTypename<TInlineAllocator<N>>
+{
+	inline static constexpr std::string_view Value = Concat<InlineAllocatorPrefix, HexString<N>>;
+};
+
+template<int N>
+struct TShortTypename<TInlineSetAllocator<N>> : TShortTypename<TInlineAllocator<N>> {};
+
+template <typename T, typename KeyFuncs, typename SetAllocator>
+struct TTypename<TSet<T, KeyFuncs, SetAllocator>>
+{
+	inline static constexpr std::string_view RangeBindName = Concat<UeSetName, ShortTypename<KeyFuncs>, ShortTypename<SetAllocator>>;
+};
+
+template <typename K, typename V, typename SetAllocator, typename KeyFuncs>
+struct TTypename<TMap<K, V, SetAllocator, KeyFuncs>>
+{
+	inline static constexpr std::string_view RangeBindName = Concat<UeMapName, ShortTypename<SetAllocator>, ShortTypename<KeyFuncs>>;
+};
+
 template<>
 PLAINPROPS_API void AppendString(FString& Out, const FName& Name);
 
@@ -722,7 +787,6 @@ struct TRangeBind<TOptional<T>>
 	using Type = UE::TOptionalBinding<T>;
 };
 
-
 template<>
 struct TCustomBind<FTransform>
 {
diff --git a/Engine/Plugins/Experimental/PlainProps/Source/Public/PlainPropsWrite.h b/Engine/Plugins/Experimental/PlainProps/Source/Public/PlainPropsWrite.h
--- a/Engine/Plugins/Experimental/PlainProps/Source/Public/PlainPropsWrite.h
+++ b/Engine/Plugins/Experimental/PlainProps/Source/Public/PlainPropsWrite.h
@@ -8,6 +8,7 @@
 
 namespace PlainProps
 {
+
 struct FBuiltSchemas;
 struct FBuiltStruct;
 class FDebugIds;
@@ -15,13 +16,14 @@ class FIdIndexerBase;
 class FNestedScopeIndexer;
 class FParametricTypeIndexer;
 struct FWriteIds;
+struct IStructBindIds;
 
 enum class ESchemaFormat { StableNames, InMemoryNames };
 
 class FWriter
 {
 public:
-	PLAINPROPS_API FWriter(const FIdIndexerBase& DeclaredIds, const FBuiltSchemas& InSchemas, ESchemaFormat Format);
+	PLAINPROPS_API FWriter(const FIdIndexerBase& DeclaredIds, const IStructBindIds& BindIds, const FBuiltSchemas& InSchemas, ESchemaFormat Format);
 	PLAINPROPS_API ~FWriter();
 	
 	PLAINPROPS_API bool								Uses(FNameId BuiltId) const;
diff --git a/Engine/Plugins/Experimental/PlainProps/Source/Tests/PlainPropsCttiTests.cpp b/Engine/Plugins/Experimental/PlainProps/Source/Tests/PlainPropsCttiTests.cpp
--- a/Engine/Plugins/Experimental/PlainProps/Source/Tests/PlainPropsCttiTests.cpp
+++ b/Engine/Plugins/Experimental/PlainProps/Source/Tests/PlainPropsCttiTests.cpp
@@ -7,11 +7,6 @@
 namespace PlainProps::Test
 {
 
-static_assert(std::is_same_v<CttiOf<float>::Type, float>);
-static_assert(std::string_view(CttiOf<int>::Name) == std::string_view("s32"));
-
-//////////////////////////////////////////////////////////////////////////
-
 enum class E1 : uint8_t { A, B };
 
 namespace Actual
@@ -38,7 +33,6 @@ static_assert(Actual::E1_Ctti::Enumerators[1].Constant	== Expect::E1_Ctti::Enume
 static_assert(std::string_view(Actual::E1_Ctti::Enumerators[0].Name) == std::string_view(Expect::E1_Ctti::Enumerators[0].Name));
 static_assert(std::string_view(Actual::E1_Ctti::Enumerators[1].Name) == std::string_view(Expect::E1_Ctti::Enumerators[1].Name));
 
-
 //////////////////////////////////////////////////////////////////////////
 
 struct S1
diff --git a/Engine/Plugins/Experimental/PlainProps/Source/Tests/PlainPropsStringTests.cpp b/Engine/Plugins/Experimental/PlainProps/Source/Tests/PlainPropsStringTests.cpp
--- a/Engine/Plugins/Experimental/PlainProps/Source/Tests/PlainPropsStringTests.cpp
+++ b/Engine/Plugins/Experimental/PlainProps/Source/Tests/PlainPropsStringTests.cpp
@@ -9,17 +9,17 @@ namespace PlainProps::Private
 	static constexpr std::string_view Hello = "Hello";
 	static constexpr std::string_view Space = " ";
 	static constexpr std::string_view World = "World";
-	static_assert("Hello World"sv == Private::Concat<Hello, Space, World>);
+	static_assert("Hello World"sv == Concat<Hello, Space, World>);
 	
-	static_assert("0"sv == Private::HexString<0>);
-	static_assert("1"sv == Private::HexString<1>);
-	static_assert("9"sv == Private::HexString<9>);
-	static_assert("A"sv == Private::HexString<0xA>);
-	static_assert("F"sv == Private::HexString<0xF>);
-	static_assert("10"sv == Private::HexString<0x10>);
-	static_assert("FF"sv == Private::HexString<0xFF>);
-	static_assert("100"sv == Private::HexString<0x100>);
-	static_assert("FFF"sv == Private::HexString<0xFFF>);
-	static_assert("1000"sv == Private::HexString<0x1000>);
-	static_assert("FEDCBA9876543210"sv == Private::HexString<0xFEDCBA9876543210>);
+	static_assert("0"sv == HexString<0>);
+	static_assert("1"sv == HexString<1>);
+	static_assert("9"sv == HexString<9>);
+	static_assert("A"sv == HexString<0xA>);
+	static_assert("F"sv == HexString<0xF>);
+	static_assert("10"sv == HexString<0x10>);
+	static_assert("FF"sv == HexString<0xFF>);
+	static_assert("100"sv == HexString<0x100>);
+	static_assert("FFF"sv == HexString<0xFFF>);
+	static_assert("1000"sv == HexString<0x1000>);
+	static_assert("FEDCBA9876543210"sv == HexString<0xFEDCBA9876543210>);
 }
\ No newline at end of file
diff --git a/Engine/Plugins/Experimental/PlainProps/Source/Tests/PlainPropsTests.cpp b/Engine/Plugins/Experimental/PlainProps/Source/Tests/PlainPropsTests.cpp
--- a/Engine/Plugins/Experimental/PlainProps/Source/Tests/PlainPropsTests.cpp
+++ b/Engine/Plugins/Experimental/PlainProps/Source/Tests/PlainPropsTests.cpp
@@ -104,7 +104,7 @@ TEST_CASE_NAMED(FPlainPropsIndexTest, "System::Core::Serialization::PlainProps::
 
 inline constexpr uint32 TestMagics[] = {0xFEEDF00D, 0xABCD1234, 0xDADADAAA, 0x99887766, 0xF0F1F2F3 };
 
-class FTestBatchBuilder : public TIdIndexer<FAnsiString>
+class FTestBatchBuilder : public TIdIndexer<FAnsiString>, public IStructBindIds
 {
 public:
 	FTestBatchBuilder(FScratchAllocator& InScratch) : Declarations(*this), Scratch(InScratch) {}
@@ -151,6 +151,12 @@ private:
 	}
 
 	TArray<char> GetNameData() const;
+
+	virtual FStructSchemaId GetDeclId(FStructSchemaId BindId) const override
+	{
+		checkf(false, TEXT("All struct ids should be declared, nothing is bound with different names in this test suite"));
+		return BindId;
+	}
 };
 
 FStructSchemaId FTestBatchBuilder::DeclareStruct(FTypeId Type, std::initializer_list<const char*> MemberOrder, EMemberPresence Occupancy, FOptionalStructSchemaId Super)
@@ -185,7 +191,7 @@ void FTestBatchBuilder::AddObject(FStructSchemaId Schema, FMemberBuilder&& Membe
 TArray64<uint8> FTestBatchBuilder::Write()
 {
 	// Build partial schemas
-	FSchemasBuilder SchemaBuilders(Declarations, Scratch);
+	FSchemasBuilder SchemaBuilders(Declarations, *this, Scratch);
 	for (const TPair<FStructSchemaId, FBuiltStructPtr>& Object : Objects)
 	{
 		SchemaBuilders.NoteStructAndMembers(Object.Key, *Object.Value);
@@ -193,7 +199,7 @@ TArray64<uint8> FTestBatchBuilder::Write()
 	FBuiltSchemas Schemas = SchemaBuilders.Build(); 
 
 	// Filter out declared but unused names and ids
-	FWriter Writer(*this, Schemas, ESchemaFormat::StableNames);
+	FWriter Writer(*this, *this, Schemas, ESchemaFormat::StableNames);
 
 	// Write names
 	TArray64<uint8> Out;
diff --git a/Engine/Plugins/Experimental/PlainProps/Source/Tests/PlainPropsUeCoreTests.cpp b/Engine/Plugins/Experimental/PlainProps/Source/Tests/PlainPropsUeCoreTests.cpp
--- a/Engine/Plugins/Experimental/PlainProps/Source/Tests/PlainPropsUeCoreTests.cpp
+++ b/Engine/Plugins/Experimental/PlainProps/Source/Tests/PlainPropsUeCoreTests.cpp
@@ -101,36 +101,41 @@ struct TScopedStructDeclaration
 {
 	using Ids = typename Runtime::Ids;
 
-	FStructSchemaId Id;
+	FStructSchemaId DeclId;
 
 	TScopedStructDeclaration()
-	: Id(DeclareNativeStruct<CttiOf<T>, Ids>(Runtime::GetTypes(), Occupancy))
+	: DeclId(DeclareNativeStruct<CttiOf<T>, Ids>(Runtime::GetTypes(), Occupancy))
 	{}
 
 	~TScopedStructDeclaration()
 	{
-		Runtime::GetTypes().DropStruct(Id);
+		Runtime::GetTypes().DropStructRef(DeclId);
 	}
 
 	const FStructDeclaration& Get() const
 	{
-		return Runtime::GetTypes().Get(Id);
+		return Runtime::GetTypes().Get(DeclId);
 	}
 };
 
 template<class T, EMemberPresence Occupancy = EMemberPresence::AllowSparse, class Runtime = FDefaultRuntime>
 struct TScopedStructBinding : TScopedStructDeclaration<T, Occupancy, Runtime>
 {
-	using TScopedStructDeclaration<T, Occupancy, Runtime>::Id;
+	using Super = TScopedStructDeclaration<T, Occupancy, Runtime>;
+	using Ids = typename Runtime::Ids;
+	using Typename = TTypename<T>;
+
+	FStructSchemaId BindId;
 
 	TScopedStructBinding()
+	: BindId(IndexStructBindIdIfNeeded<T, Ids>(Super::DeclId))
 	{
-		BindNativeStruct<CttiOf<T>, Runtime>(Runtime::GetSchemas(), Id);
+		BindNativeStruct<CttiOf<T>, Runtime>(Runtime::GetSchemas(), BindId, Super::DeclId);
 	}
 
 	~TScopedStructBinding()
 	{
-		Runtime::GetSchemas().DropStruct(Id);
+		Runtime::GetSchemas().DropStruct(BindId);
 	}
 };
 
@@ -142,7 +147,7 @@ struct FNameDeclaration
 	FStructSchemaId		Id;
 	FMemberId			Idx;
 
-	FNameDeclaration(FTypeId Type = FIds::IndexNativeType(CttiOf<FName>::Name))
+	FNameDeclaration(FTypeId Type = IndexTypename<FIds, ETypename::Decl, TTypename<FName>>())
 	: Id(FIds::IndexStruct(Type))
 	, Idx(FIds::IndexMember("Idx"))
 	{
@@ -151,13 +156,13 @@ struct FNameDeclaration
 
 	~FNameDeclaration()
 	{
-		GTypes.DropStruct(Id);
+		GTypes.DropStructRef(Id);
 	}
 };
 
 struct FTestCustomBinding : ICustomBinding
 {
-	virtual FStructSchemaId GetId() const = 0;
+//	virtual FDualStructSchemaId GetId() const = 0;
 };
 
 struct FNameBinding : FTestCustomBinding
@@ -181,7 +186,7 @@ struct FNameBinding : FTestCustomBinding
 		return A.IsEqual(B, ENameCase::CaseSensitive);
 	}
 
-	virtual FStructSchemaId GetId() const override { return Declaration.Id; }
+//	virtual FDualStructSchemaId GetId() const override { return { Declaration.Id, Declaration.Id }; }
 
 	FNameDeclaration	Declaration;
 	TSet<FName>			Names;
@@ -212,23 +217,23 @@ private:
 FBatchSaver::FBatchSaver()
 : Customs(/* debug */ GNames, &GCustoms)
 {
-	Customs.BindStruct(SavedNames.Declaration.Id, SavedNames);
+	Customs.BindStruct(SavedNames.Declaration.Id, SavedNames.Declaration.Id, SavedNames);
 }
 
 template<class T>
 void FBatchSaver::Save(T&& Object) 
 {
-	FStructSchemaId Id = IndexStruct<std::remove_reference_t<T>, FIds>();
-	SavedObjects.Emplace(Id, SaveStruct(&Object, Id, {GTypes, GSchemas, Customs, Scratch}));
+	FDualStructSchemaId Id = IndexStructDualId<std::remove_reference_t<T>, FIds>();
+	SavedObjects.Emplace(Id.Decl, SaveStruct(&Object, Id.Bind, {GTypes, GSchemas, Customs, Scratch}));
 }
 
 template<class T>
 bool FBatchSaver::SaveDelta(const T& Object, const T& Default) 
 {
-	FStructSchemaId Id = IndexStruct<std::remove_reference_t<T>, FIds>();
-	if (FBuiltStructPtr Delta = SaveStructDelta(&Object, &Default, Id, {GTypes, GSchemas, Customs, Scratch}))
+	FDualStructSchemaId Id = IndexStructDualId<std::remove_reference_t<T>, FIds>();
+	if (FBuiltStructPtr Delta = SaveStructDelta(&Object, &Default, Id.Bind, {GTypes, GSchemas, Customs, Scratch}))
 	{
-		SavedObjects.Emplace(Id, MoveTemp(Delta));
+		SavedObjects.Emplace(Id.Decl, MoveTemp(Delta));
 		return true;
 	}
 	return false;
@@ -253,7 +258,7 @@ inline constexpr uint32 Magics[] = { 0xFEEDF00D, 0xABCD1234, 0xDADADAAA, 0x99887
 TArray64<uint8> FBatchSaver::Write() const
 {
 	// Build partial schemas
-	FSchemasBuilder SchemaBuilders(GTypes, Scratch);
+	FSchemasBuilder SchemaBuilders(GTypes, GSchemas, Scratch);
 	for (const IdBuiltStructPair& Object : SavedObjects)
 	{
 		SchemaBuilders.NoteStructAndMembers(Object.Key, *Object.Value);
@@ -261,7 +266,7 @@ TArray64<uint8> FBatchSaver::Write() const
 	FBuiltSchemas Schemas = SchemaBuilders.Build(); 
 
 	// Filter out declared but unused names and ids
-	FWriter Writer(GNames, Schemas, ESchemaFormat::StableNames);
+	FWriter Writer(GNames, GSchemas, Schemas, ESchemaFormat::StableNames);
 	TArray<FName> UsedNames;
 	for (uint32 Idx = 0, Num = GNames.NumNames(); Idx < Num; ++Idx)
 	{
@@ -346,7 +351,7 @@ public:
 
 		// Read names and bind custom loader
 		Names.Names.Append(GrabNumAndArray<FName>(It));
-		Customs.BindStruct(Names.Declaration.Id, Names);
+		Customs.BindStruct(Names.Declaration.Id, Names.Declaration.Id, Names);
 		CHECK(It.Grab<uint32>() == Magics[5]);
 
 		// Finally create load plans
@@ -532,6 +537,61 @@ inline bool operator==(const FMaps& A, const FMaps& B)
 	return LegacyCompareEqual(A.Leaves, B.Leaves) && LegacyCompareEqual(A.Ranges, B.Ranges) && LegacyCompareEqual(A.Structs, B.Structs);
 }
 
+struct FIntAlias
+{
+	int X;
+	bool operator==(const FIntAlias&) const = default;
+};
+
+struct FSame1
+{
+	int X = 1;
+	bool operator==(const FSame1&) const = default;
+};
+
+struct FSame2
+{ 
+	int Unused;
+	int X = 2;
+	bool operator==(const FSame2& O) const { return X == O.X; };
+};
+
+
+PP_REFLECT_STRUCT(PlainProps::UE::Test, FIntAlias, void, X);
+PP_REFLECT_STRUCT(PlainProps::UE::Test, FSame1, void, X);
+PP_REFLECT_STRUCT(PlainProps::UE::Test, FSame2, void, X);
+
+} namespace PlainProps {
+
+template <> struct TTypename<PlainProps::UE::Test::FIntAlias>
+{
+	inline static constexpr std::string_view DeclName = "FInt";
+	inline static constexpr std::string_view BindName = "IntAlias";
+};
+template <> struct TTypename<PlainProps::UE::Test::FSame1>
+{
+	inline static constexpr std::string_view DeclName = "Same";
+	inline static constexpr std::string_view BindName = "Same1";
+};
+template <> struct TTypename<PlainProps::UE::Test::FSame2>
+{
+	inline static constexpr std::string_view DeclName = "Same";
+	inline static constexpr std::string_view BindName = "Same2";
+};
+
+} namespace PlainProps::UE::Test {
+
+struct FTypeErasure
+{
+	FSame1 A;
+	FSame2 B;
+	FIntAlias C;
+	TPair<FString, TArray<char8_t>> D;
+	TPair<TArray<char>, TArray<char, TInlineAllocator<8>>> E;
+	bool operator==(const FTypeErasure&) const = default;
+};
+PP_REFLECT_STRUCT(PlainProps::UE::Test, FTypeErasure, void, A, B, C, D, E); 
+
 //////////////////////////////////////////////////////////////////////////
 
 struct FUniquePtrs
@@ -674,6 +734,18 @@ PP_REFLECT_STRUCT(PlainProps::UE::Test, FDelta, void, A, B, C, D, E);
 
 //////////////////////////////////////////////////////////////////////////
 
+template<typename T>
+TArray<T> MakeArray(const T* Str)
+{
+	return TArray<T>(Str, TCString<T>::Strlen(Str));
+}
+
+template<int N, typename T>
+TArray<T, TInlineAllocator<N>> MakeInlArray(const T* Str)
+{
+	return TArray<T, TInlineAllocator<N>>(Str, TCString<T>::Strlen(Str));
+}
+
 TEST_CASE_NAMED(FPlainPropsUeCoreTest, "System::Core::Serialization::PlainProps::UE::Core", "[Core][PlainProps][SmokeFilter]")
 {
 	SECTION("Basic")
@@ -828,7 +900,28 @@ TEST_CASE_NAMED(FPlainPropsUeCoreTest, "System::Core::Serialization::PlainProps:
 				}
 			});	
 	}
-		
+
+	SECTION("TypeErasure")
+	{
+		TScopedStructBinding<FSame1> Same1;
+		TScopedStructBinding<FSame2> Same2;	
+		TScopedStructBinding<FIntAlias> IntAlias;
+		TScopedStructBinding<TPair<FString, TArray<char8_t>>> X;
+		TScopedStructBinding<TPair<TArray<char>, TArray<char, TInlineAllocator<8>>>> Y;
+		TScopedStructBinding<FTypeErasure> TypeErasure;
+
+		Run([](FBatchSaver& Batch)
+			{
+				Batch.Save(FTypeErasure{});
+				Batch.Save(FTypeErasure{{10}, {0,20}, {30}, {"a", TArray<char8_t>({u8'b'})}, {MakeArray("c"), MakeInlArray<8>("d")}});
+			}, 
+			[](FBatchLoader& Batch)
+			{
+				CHECK(Batch.Load<FTypeErasure>() == FTypeErasure{});
+				CHECK(Batch.Load<FTypeErasure>() == FTypeErasure{FTypeErasure{{10}, {0,20}, {30}, {"a", TArray<char8_t>({u8'b'})}, {MakeArray("c"), MakeInlArray<8>("d")}}});
+			});
+	}
+
 	SECTION("TSet")
 	{
 		TScopedStructBinding<FInt> Int;
