PlainProps finish namespace support

#rnx

diff --git a/Engine/Plugins/Experimental/PlainProps/Source/Public/PlainPropsBind.h b/Engine/Plugins/Experimental/PlainProps/Source/Public/PlainPropsBind.h
--- a/Engine/Plugins/Experimental/PlainProps/Source/Public/PlainPropsBind.h
+++ b/Engine/Plugins/Experimental/PlainProps/Source/Public/PlainPropsBind.h
@@ -569,30 +569,55 @@ private:
 
 ////////////////////////////////////////////////////////////////////////////////////////////////
 
-template<class Ids, ETypename Kind, typename Typename>
-FTypeId IndexBaseName()
+// Cached by function static
+template<class Ids, const char* Scope>
+FScopeId GetScopeId()
+{
+	static FScopeId Id = Ids::IndexScope(Scope);
+	return Id;
+}
+
+// Cached by function static
+template<class Ids, typename Typename>
+FScopeId GetNamespaceId()
+{
+	if constexpr (Typename::Namespace.size())
+	{
+		return GetScopeId<Ids, Typename::Namespace.data()>();
+	}
+	else
+	{
+		return NoId;
+	}
+}
+
+template<ETypename Kind, typename Typename>
+constexpr std::string_view SelectStructName()
 {	
 	if constexpr (Kind == ETypename::Bind && ExplicitBindName<Typename>)
 	{
-		return Ids::IndexNativeType(ToAnsiView(Typename::BindName));
+		return Typename::BindName;
 	}
 	else
 	{
-		return Ids::IndexNativeType(ToAnsiView(Typename::DeclName));	
+		return Typename::DeclName;	
 	}
 }
 
 template<class Ids, ETypename Kind, typename Typename>
-FTypeId IndexTypename()
+FTypeId IndexStructName()
 {
-	return IndexBaseName<Ids, Kind, Typename>();
-}
-
-template<class Ids, ETypename Kind, ParametricName Typename>
-FTypeId IndexTypename()
-{
-	FTypeId BaseName = IndexBaseName<Ids, Kind, Typename>();
-	return IndexParametricType<Ids, Kind>(BaseName, (typename Typename::Parameters*)nullptr);
+	FTypeId BaseName = { GetNamespaceId<Ids, Typename>(), 
+						 Ids::IndexTypename(ToAnsiView(SelectStructName<Kind, Typename>())) };
+	
+	if constexpr (ParametricName<Typename>)
+	{
+		return IndexParametricType<Ids, Kind>(BaseName, (typename Typename::Parameters*)nullptr);
+	}
+	else
+	{
+		return BaseName;
+	}
 }
 
 template<typename Struct, class Ids>
@@ -602,7 +627,7 @@ FStructSchemaId IndexStructBindIdIfNeeded(FStructSchemaId DeclId)
 
 	if constexpr (ExplicitBindName<Typename> || ParametricName<Typename>)
 	{
-		return Ids::IndexStruct(IndexTypename<Ids, ETypename::Bind, Typename>());
+		return Ids::IndexStruct(IndexStructName<Ids, ETypename::Bind, Typename>());
 	}
 	else
 	{
@@ -616,13 +641,13 @@ FDualStructSchemaId IndexStructDualId()
 	using Typename = TTypename<Struct>;
 	
 	FDualStructSchemaId Out;
-	FTypeId DeclName = IndexTypename<Ids, ETypename::Decl, Typename>();
+	FTypeId DeclName = IndexStructName<Ids, ETypename::Decl, Typename>();
 	Out.Decl = Ids::IndexStruct(DeclName);
 	Out.Bind = Out.Decl;
 
 	if constexpr (ExplicitBindName<Typename> || ParametricName<Typename>)
 	{
-		FTypeId BindName = IndexTypename<Ids, ETypename::Bind, Typename>();
+		FTypeId BindName = IndexStructName<Ids, ETypename::Bind, Typename>();
 		Out.Bind = BindName != DeclName ? Ids::IndexStruct(BindName) : Out.Decl;
 	}
 
@@ -633,7 +658,7 @@ FDualStructSchemaId IndexStructDualId()
 template<class Ids, typename Struct>
 FStructSchemaId GetStructDeclId()
 {
-	static FStructSchemaId Id = Ids::IndexStruct(IndexTypename<Ids, ETypename::Decl, TTypename<Struct>>());
+	static FStructSchemaId Id = Ids::IndexStruct(IndexStructName<Ids, ETypename::Decl, TTypename<Struct>>());
 	return Id;
 }
 
@@ -641,15 +666,27 @@ FStructSchemaId GetStructDeclId()
 template<class Ids, typename Struct>
 FStructSchemaId GetStructBindId()
 {
-	static FStructSchemaId Id = Ids::IndexStruct(IndexTypename<Ids, ETypename::Bind, TTypename<Struct>>());
+	static FStructSchemaId Id = Ids::IndexStruct(IndexStructName<Ids, ETypename::Bind, TTypename<Struct>>());
 	return Id;
 }
 
+template<class Ids, typename Ctti>
+FTypeId IndexCttiName()
+{
+	FTypenameId Name = Ids::IndexTypename(Ctti::Name);
+	FScopeId Namespace = NoId;
+	if constexpr (Ctti::Namespace[0] != '\0')
+	{
+		Namespace = GetScopeId<Ids, Ctti::Namespace>();
+	}
+	return { Namespace, Name };
+}
+
 // Cached by function static
 template<class Ids, typename Enum>
 FEnumSchemaId GetEnumId()
 {
-	static FEnumSchemaId Id = Ids::IndexEnum(CttiOf<Enum>::Name);
+	static FEnumSchemaId Id = Ids::IndexEnum(IndexCttiName<Ids, CttiOf<Enum>>());
 	return Id;
 }
 
@@ -666,10 +703,10 @@ FTypeId IndexParameterName()
 	return IndexArithmeticName<Ids, Leaf>();
 }
 
-template<class Ids, ETypename, Enumeration Leaf>
+template<class Ids, ETypename, Enumeration Enum>
 FTypeId IndexParameterName()
 {
-	return Ids::IndexEnum(CttiOf<Leaf>::Name);
+	return GetEnumId<Ids, Enum>();
 }
 
 template<class Ids, ETypename Kind, typename T>
@@ -678,7 +715,7 @@ FTypeId IndexParameterName()
 	using RangeBinding = RangeBind<T>;
 	if constexpr (std::is_void_v<RangeBinding>)
 	{
-		return IndexTypename<Ids, Kind, TTypename<T>>();
+		return IndexStructName<Ids, Kind, TTypename<T>>();
 	}
 	else
 	{
@@ -692,7 +729,8 @@ FTypeId IndexParameterName()
 		}
 		else
 		{
-			FTypeId RangeBindName = Ids::IndexNativeType(ToAnsiView(TTypename<T>::RangeBindName));
+			using Typename = TTypename<T>;
+			FTypeId RangeBindName = { GetNamespaceId<Ids, Typename>(), Ids::IndexTypename(ToAnsiView(Typename::RangeBindName)) };
 			return Ids::GetIndexer().MakeParametricType(RangeBindName, {ItemTypename, SizeTypename});
 		}
 	}
@@ -717,7 +755,7 @@ FStructSchemaId BindCustomStructOnce()
 
 		FBinding()
 		{
-			FTypeId DeclType = IndexTypename<Ids, ETypename::Decl, TTypename<Type>>();
+			FTypeId DeclType = IndexStructName<Ids, ETypename::Decl, TTypename<Type>>();
 			DeclId = Ids::IndexStruct(DeclType);
 			BindId = IndexStructBindIdIfNeeded<Type, Ids>(DeclId);
 			
@@ -903,7 +941,7 @@ FEnumSchemaId DeclareNativeEnum(FDeclarations& Out, EEnumMode Mode)
 {
 	using UnderlyingType = std::underlying_type_t<typename Ctti::Type>;
 
-	FTypeId Type = Ids::IndexNativeType(Ctti::Name);
+	FTypeId Type = IndexCttiName<Ids, Ctti>();
 	FEnumSchemaId Id = Ids::IndexEnum(Type);
 	FEnumerator Enumerators[Ctti::NumEnumerators];
 	for (FEnumerator& Enumerator : Enumerators)
@@ -922,7 +960,7 @@ FStructSchemaId DeclareNativeStruct(FDeclarations& Out, EMemberPresence Occupanc
 	using Typename = TTypename<typename Ctti::Type>;
 	using SuperType = typename Ctti::Super;
 
-	FTypeId Type = IndexTypename<Ids, ETypename::Decl, Typename>();
+	FTypeId Type = IndexStructName<Ids, ETypename::Decl, Typename>();
 	FStructSchemaId Id = Ids::IndexStruct(Type);
 	FOptionalStructSchemaId SuperId;
 	if constexpr (!std::is_void_v<SuperType>)
diff --git a/Engine/Plugins/Experimental/PlainProps/Source/Public/PlainPropsCtti.h b/Engine/Plugins/Experimental/PlainProps/Source/Public/PlainPropsCtti.h
--- a/Engine/Plugins/Experimental/PlainProps/Source/Public/PlainPropsCtti.h
+++ b/Engine/Plugins/Experimental/PlainProps/Source/Public/PlainPropsCtti.h
@@ -30,7 +30,7 @@ template<typename T> using CttiOf = typename TCttiOf<T>::Type;
 #define _PP_REFLECT_THINGS(N, MACRO, ...)				_PP_EXPAND(_PP_REFLECT_##N(N, MACRO, __VA_ARGS__))
 
 #define _PP_REFLECT_ENUMERATOR(N, NS, T, C)				{ #C, Type::C }, 
-#define _PP_REFLECT_ENUM_INNER(N, NS, T, ES)			struct T##_Ctti { inline static constexpr char Name[] = #T; using Type = NS :: T; static constexpr int NumEnumerators = N; static constexpr struct { const char* Name; Type Constant; } Enumerators[] = { ES }; }; T##_Ctti CttiOfPtr(T*);
+#define _PP_REFLECT_ENUM_INNER(N, NS, T, ES)			struct T##_Ctti { inline static constexpr char Name[] = #T; using Type = NS :: T; static constexpr int NumEnumerators = N; static constexpr struct { const char* Name; Type Constant; } Enumerators[] = { ES }; inline static constexpr char Namespace[] = #NS;  }; T##_Ctti CttiOfPtr(T*);
 #define _PP_REFLECT_ENUM(N, NS, T, ...)					_PP_REFLECT_ENUM_INNER(N, NS, T, _PP_EXPAND(_PP_REFLECT_THINGS(N, _PP_REFLECT_ENUMERATOR, NS, T, __VA_ARGS__)))					
 #define _PP_REFLECT_STRUCT(N, NS, T, S, ...)			PP_REFLECT_STRUCT_ONLY(N, NS, T, S)				_PP_EXPAND(_PP_REFLECT_THINGS(N, PP_REFLECT_MEMBER, NS, T, __VA_ARGS__))
 #define _PP_REFLECT_STRUCT_TEMPLATE(N, NS, T, S, ...)	PP_REFLECT_STRUCT_TEMPLATE_ONLY(N, NS, T, S)	_PP_EXPAND(_PP_REFLECT_THINGS(N, PP_REFLECT_TEMPLATE_MEMBER, NS, T, __VA_ARGS__))
@@ -40,13 +40,13 @@ template<typename T> using CttiOf = typename TCttiOf<T>::Type;
 #define PP_REFLECT_ENUM(NS, T, ...)						_PP_REFLECT_ENUM(_PP_NUM_ARGS(__VA_ARGS__), NS, T, __VA_ARGS__)
 #define PP_REFLECT_STRUCT(NS, T, S, ...)				_PP_REFLECT_STRUCT(_PP_NUM_ARGS(__VA_ARGS__), NS, T, S, __VA_ARGS__)
 #define PP_REFLECT_STRUCT_TEMPLATE(NS, T, S, ...)		_PP_REFLECT_STRUCT_TEMPLATE(_PP_NUM_ARGS(__VA_ARGS__), NS, T, S, __VA_ARGS__)
-#define PP_NAME_STRUCT(NS, T)							struct T##_Ctti { inline static constexpr char Name[] = #T; using Type = NS :: T; }; T##_Ctti CttiOfPtr(T*);
-#define PP_NAME_STRUCT_TEMPLATE(NS, T)					template<class... Ts> struct T##_Ctti { inline static constexpr char Name[] = #T; using Type = NS :: T<Ts...>; using TemplateArgs = std::tuple<Ts...>; }; template<class... Ts> T##_Ctti<Ts...> CttiOfPtr(T<Ts...>*);
+#define PP_NAME_STRUCT(NS, T)							struct T##_Ctti { inline static constexpr char Name[] = #T; using Type = NS :: T; inline static constexpr char Namespace[] = #NS; }; T##_Ctti CttiOfPtr(T*);
+#define PP_NAME_STRUCT_TEMPLATE(NS, T)					template<class... Ts> struct T##_Ctti { inline static constexpr char Name[] = #T; using Type = NS :: T<Ts...>; using TemplateArgs = std::tuple<Ts...>; inline static constexpr char Namespace[] = #NS; }; template<class... Ts> T##_Ctti<Ts...> CttiOfPtr(T<Ts...>*);
 		
 // Alternate set of macros to reflect classes/structs with bitfield bool members, e.g. uint8 bOol : 1;
 
-#define PP_REFLECT_STRUCT_ONLY(N, NS, T, S)				struct T##_Ctti { inline static constexpr char Name[] = #T; using Type = NS :: T; using Super = S; static constexpr int NumVars = N; template<int> struct Var; }; T##_Ctti CttiOfPtr(T*);
-#define PP_REFLECT_STRUCT_TEMPLATE_ONLY(N, NS, T, S)	template<class... Ts> struct T##_Ctti { inline static constexpr char Name[] = #T; using Type = NS :: T<Ts...>; using Super = S; using TemplateArgs = std::tuple<Ts...>; static constexpr int NumVars = N; template<int, int> struct Var_; template<int I> using Var = Var_<I, 0>; }; template<class... Ts> T##_Ctti<Ts...> CttiOfPtr(T<Ts...>*);
+#define PP_REFLECT_STRUCT_ONLY(N, NS, T, S)				struct T##_Ctti { inline static constexpr char Name[] = #T; using Type = NS :: T; using Super = S; static constexpr int NumVars = N; template<int> struct Var; inline static constexpr char Namespace[] = #NS; }; T##_Ctti CttiOfPtr(T*);
+#define PP_REFLECT_STRUCT_TEMPLATE_ONLY(N, NS, T, S)	template<class... Ts> struct T##_Ctti { inline static constexpr char Name[] = #T; using Type = NS :: T<Ts...>; using Super = S; using TemplateArgs = std::tuple<Ts...>; static constexpr int NumVars = N; template<int, int> struct Var_; template<int I> using Var = Var_<I, 0>; inline static constexpr char Namespace[] = #NS; }; template<class... Ts> T##_Ctti<Ts...> CttiOfPtr(T<Ts...>*);
 #define PP_REFLECT_MEMBER(N, NS, T, M)					template<> struct T##_Ctti::Var<N> { inline static constexpr char Name[] = #M; using Type = decltype(NS :: T :: M); static constexpr std::size_t Offset = offsetof(NS :: T, M); static constexpr auto Pointer = &NS :: T :: M; static constexpr int Index = N; };
 #define PP_REFLECT_BIT_MEMBER(N, NS, T, M, BYTEOFFSET, BITOFFSET) template<> struct T##_Ctti::Var<N> { inline static constexpr char Name[] = #M; using Type = decltype(NS :: T :: M); static constexpr std::size_t Offset = BYTEOFFSET; static constexpr int BitIndex = BITOFFSET; static constexpr int Index = N; };
 #define PP_REFLECT_TEMPLATE_MEMBER(N, NS, T, M)			template<class... Ts> template <int _> struct T##_Ctti<Ts...>::Var_<N, _> { inline static constexpr char Name[] = #M; using Type = decltype(NS :: T<Ts...> :: M); static constexpr std::size_t Offset = offsetof(NS :: T<Ts...>, M); static constexpr auto Pointer = &NS :: T<Ts...> :: M; static constexpr int Index = N; };
diff --git a/Engine/Plugins/Experimental/PlainProps/Source/Public/PlainPropsTypename.h b/Engine/Plugins/Experimental/PlainProps/Source/Public/PlainPropsTypename.h
--- a/Engine/Plugins/Experimental/PlainProps/Source/Public/PlainPropsTypename.h
+++ b/Engine/Plugins/Experimental/PlainProps/Source/Public/PlainPropsTypename.h
@@ -35,6 +35,7 @@ template<typename T>
 struct TTypename
 {
 	inline static constexpr std::string_view DeclName = CttiOf<T>::Name;
+	inline static constexpr std::string_view Namespace = CttiOf<T>::Namespace;
 };
 
 template<class Typename>
@@ -90,6 +91,7 @@ struct TTypename<T>
 {
 	inline static constexpr FUnpackedLeafType Leaf = ReflectArithmetic<T>;
 	inline static constexpr std::string_view DeclName = ArithmeticName<Leaf.Type, Leaf.Width>;
+	inline static constexpr std::string_view Namespace;
 };
 
 } // namespace PlainProps
\ No newline at end of file
diff --git a/Engine/Plugins/Experimental/PlainProps/Source/Public/PlainPropsUeCoreBindings.h b/Engine/Plugins/Experimental/PlainProps/Source/Public/PlainPropsUeCoreBindings.h
--- a/Engine/Plugins/Experimental/PlainProps/Source/Public/PlainPropsUeCoreBindings.h
+++ b/Engine/Plugins/Experimental/PlainProps/Source/Public/PlainPropsUeCoreBindings.h
@@ -686,18 +686,21 @@ template <>
 struct TTypename<FName>
 {
 	inline static constexpr std::string_view DeclName = "Name";
+	inline static constexpr std::string_view Namespace;
 };
 
 template <>
 struct TTypename<FTransform>
 {
 	inline static constexpr std::string_view DeclName = "Transform";
+	inline static constexpr std::string_view Namespace;
 };
 
 template <typename K, typename V>
 struct TTypename<TPair<K,V>>
 {
 	inline static constexpr std::string_view DeclName = "Pair";
+	inline static constexpr std::string_view Namespace;
 	using Parameters = std::tuple<K, V>;
 };
 
@@ -705,6 +708,7 @@ template <>
 struct TTypename<FString>
 {
 	inline static constexpr std::string_view RangeBindName = "String";
+	inline static constexpr std::string_view Namespace;
 };
 
 inline static constexpr std::string_view UeArrayName = "Array";
@@ -715,6 +719,7 @@ template<typename T, typename Allocator>
 struct TTypename<TArray<T, Allocator>>
 {
 	inline static constexpr std::string_view RangeBindName = Concat<UeArrayName, ShortTypename<Allocator>>;
+	inline static constexpr std::string_view Namespace;
 };
 
 template<>
@@ -740,12 +745,14 @@ template <typename T, typename KeyFuncs, typename SetAllocator>
 struct TTypename<TSet<T, KeyFuncs, SetAllocator>>
 {
 	inline static constexpr std::string_view RangeBindName = Concat<UeSetName, ShortTypename<KeyFuncs>, ShortTypename<SetAllocator>>;
+	inline static constexpr std::string_view Namespace;
 };
 
 template <typename K, typename V, typename SetAllocator, typename KeyFuncs>
 struct TTypename<TMap<K, V, SetAllocator, KeyFuncs>>
 {
 	inline static constexpr std::string_view RangeBindName = Concat<UeMapName, ShortTypename<SetAllocator>, ShortTypename<KeyFuncs>>;
+	inline static constexpr std::string_view Namespace;
 };
 
 template<>
diff --git a/Engine/Plugins/Experimental/PlainProps/Source/Tests/PlainPropsUeCoreTests.cpp b/Engine/Plugins/Experimental/PlainProps/Source/Tests/PlainPropsUeCoreTests.cpp
--- a/Engine/Plugins/Experimental/PlainProps/Source/Tests/PlainPropsUeCoreTests.cpp
+++ b/Engine/Plugins/Experimental/PlainProps/Source/Tests/PlainPropsUeCoreTests.cpp
@@ -33,12 +33,9 @@ struct FIds
 	static FNameId			IndexName(FAnsiStringView Name)				{ return GNames.MakeName(FName(Name)); }
 	static FMemberId		IndexMember(FAnsiStringView Name)			{ return GNames.NameMember(FName(Name)); }
 	static FTypenameId		IndexTypename(FAnsiStringView Name)			{ return GNames.MakeTypename(FName(Name)); }
-	static FScopeId			IndexNativeScope()							{ return GNames.MakeScope(FName(UE_MODULE_NAME)); }
-	static FTypeId			IndexNativeType(FAnsiStringView Typename)	{ return {IndexNativeScope(), IndexTypename(Typename)}; }
+	static FScopeId			IndexScope(FAnsiStringView Name)			{ return GNames.MakeScope(FName(Name)); }
 	static FEnumSchemaId	IndexEnum(FTypeId Type)						{ return GNames.IndexEnum(Type); }
-	static FEnumSchemaId	IndexEnum(FAnsiStringView Name)				{ return IndexEnum(IndexNativeType(Name)); }
 	static FStructSchemaId	IndexStruct(FTypeId Type)					{ return GNames.IndexStruct(Type); }
-	static FStructSchemaId	IndexStruct(FAnsiStringView Name)			{ return IndexStruct(IndexNativeType(Name)); }
 	static FIdIndexerBase&	GetIndexer()								{ return GNames; }
 	static const FDebugIds& GetDebug()									{ return GNames; }
 };
@@ -147,7 +144,7 @@ struct FNameDeclaration
 	FStructSchemaId		Id;
 	FMemberId			Idx;
 
-	FNameDeclaration(FTypeId Type = IndexTypename<FIds, ETypename::Decl, TTypename<FName>>())
+	FNameDeclaration(FTypeId Type = IndexStructName<FIds, ETypename::Decl, TTypename<FName>>())
 	: Id(FIds::IndexStruct(Type))
 	, Idx(FIds::IndexMember("Idx"))
 	{
@@ -567,16 +564,19 @@ template <> struct TTypename<PlainProps::UE::Test::FIntAlias>
 {
 	inline static constexpr std::string_view DeclName = "FInt";
 	inline static constexpr std::string_view BindName = "IntAlias";
+	inline static constexpr std::string_view Namespace;
 };
 template <> struct TTypename<PlainProps::UE::Test::FSame1>
 {
 	inline static constexpr std::string_view DeclName = "Same";
 	inline static constexpr std::string_view BindName = "Same1";
+	inline static constexpr std::string_view Namespace;
 };
 template <> struct TTypename<PlainProps::UE::Test::FSame2>
 {
 	inline static constexpr std::string_view DeclName = "Same";
 	inline static constexpr std::string_view BindName = "Same2";
+	inline static constexpr std::string_view Namespace;
 };
 
 } namespace PlainProps::UE::Test {
