PlainProps add CustomBindingTypename and some other prep work for TSet/TMap delta bindings

#rnx

diff --git a/Engine/Plugins/Experimental/PlainProps/Source/Public/PlainPropsBind.h b/Engine/Plugins/Experimental/PlainProps/Source/Public/PlainPropsBind.h
--- a/Engine/Plugins/Experimental/PlainProps/Source/Public/PlainPropsBind.h
+++ b/Engine/Plugins/Experimental/PlainProps/Source/Public/PlainPropsBind.h
@@ -569,21 +569,14 @@ private:
 
 ////////////////////////////////////////////////////////////////////////////////////////////////
 
-// Cached by function static
-template<class Ids, const char* Scope>
-FScopeId GetScopeId()
-{
-	static FScopeId Id = Ids::IndexScope(Scope);
-	return Id;
-}
-
-// Cached by function static
 template<class Ids, typename Typename>
-FScopeId GetNamespaceId()
+FScopeId IndexNamespaceId()
 {
+	// Opt: Make cached GetNamespaceId(), either via new namespace CTTI types (maybe PP_REFLECT_NAMESPACE)
+	//		or some compile time string template parameters, perhaps a variadic template taking any number of chars
 	if constexpr (Typename::Namespace.size())
 	{
-		return GetScopeId<Ids, Typename::Namespace.data()>();
+		return Ids::IndexScope(ToAnsiView(Typename::Namespace));
 	}
 	else
 	{
@@ -607,7 +600,7 @@ constexpr std::string_view SelectStructName()
 template<class Ids, ETypename Kind, typename Typename>
 FTypeId IndexStructName()
 {
-	FTypeId BaseName = { GetNamespaceId<Ids, Typename>(), 
+	FTypeId BaseName = { IndexNamespaceId<Ids, Typename>(), 
 						 Ids::IndexTypename(ToAnsiView(SelectStructName<Kind, Typename>())) };
 	
 	if constexpr (ParametricName<Typename>)
@@ -620,13 +613,12 @@ FTypeId IndexStructName()
 	}
 }
 
-template<typename Struct, class Ids>
+template<typename Typename, class Ids>
 FStructSchemaId IndexStructBindIdIfNeeded(FStructSchemaId DeclId)
 {
-	using Typename = TTypename<Struct>;
-
 	if constexpr (ExplicitBindName<Typename> || ParametricName<Typename>)
 	{
+		// Note could pass in and reuse declared namespace here
 		return Ids::IndexStruct(IndexStructName<Ids, ETypename::Bind, Typename>());
 	}
 	else
@@ -677,7 +669,7 @@ FTypeId IndexCttiName()
 	FScopeId Namespace = NoId;
 	if constexpr (Ctti::Namespace[0] != '\0')
 	{
-		Namespace = GetScopeId<Ids, Ctti::Namespace>();
+		Namespace = Ids::IndexScope(ToAnsiView(Ctti::Namespace));
 	}
 	return { Namespace, Name };
 }
@@ -730,7 +722,7 @@ FTypeId IndexParameterName()
 		else
 		{
 			using Typename = TTypename<T>;
-			FTypeId RangeBindName = { GetNamespaceId<Ids, Typename>(), Ids::IndexTypename(ToAnsiView(Typename::RangeBindName)) };
+			FTypeId RangeBindName = { IndexNamespaceId<Ids, Typename>(), Ids::IndexTypename(ToAnsiView(Typename::RangeBindName)) };
 			return Ids::GetIndexer().MakeParametricType(RangeBindName, {ItemTypename, SizeTypename});
 		}
 	}
@@ -751,13 +743,14 @@ FStructSchemaId BindCustomStructOnce()
 	struct FBinding : CustomBinding
 	{
 		using Type = typename CustomBinding::Type;
+		using Typename = TCustomTypename<CustomBinding>;
 		using Ids = typename Runtime::Ids;
 
 		FBinding()
 		{
-			FTypeId DeclType = IndexStructName<Ids, ETypename::Decl, TTypename<Type>>();
+			FTypeId DeclType = IndexStructName<Ids, ETypename::Decl, Typename>();
 			DeclId = Ids::IndexStruct(DeclType);
-			BindId = IndexStructBindIdIfNeeded<Type, Ids>(DeclId);
+			BindId = IndexStructBindIdIfNeeded<Typename, Ids>(DeclId);
 			
 			CustomBinding::template InitIds<Ids>();
 			Runtime::GetTypes().DeclareStruct(DeclId, DeclType, CustomBinding::MemberIds, CustomBinding::Occupancy);
diff --git a/Engine/Plugins/Experimental/PlainProps/Source/Public/PlainPropsTypename.h b/Engine/Plugins/Experimental/PlainProps/Source/Public/PlainPropsTypename.h
--- a/Engine/Plugins/Experimental/PlainProps/Source/Public/PlainPropsTypename.h
+++ b/Engine/Plugins/Experimental/PlainProps/Source/Public/PlainPropsTypename.h
@@ -94,4 +94,17 @@ struct TTypename<T>
 	inline static constexpr std::string_view Namespace;
 };
 
+//////////////////////////////////////////////////////////////////////////
+
+template<class CustomBinding>
+concept WithCustomTypename = requires { typename CustomBinding::FCustomTypename; };
+
+template<class CustomBinding>
+struct TCustomTypename : TTypename<typename CustomBinding::Type>
+{};
+
+template<WithCustomTypename CustomBinding>
+struct TCustomTypename<CustomBinding> : CustomBinding::FCustomTypename
+{};
+
 } // namespace PlainProps
\ No newline at end of file
diff --git a/Engine/Plugins/Experimental/PlainProps/Source/Public/PlainPropsUeCoreBindings.h b/Engine/Plugins/Experimental/PlainProps/Source/Public/PlainPropsUeCoreBindings.h
--- a/Engine/Plugins/Experimental/PlainProps/Source/Public/PlainPropsUeCoreBindings.h
+++ b/Engine/Plugins/Experimental/PlainProps/Source/Public/PlainPropsUeCoreBindings.h
@@ -449,46 +449,44 @@ struct FTransformBinding : ICustomBinding
 
 struct FSetDeltaOps
 {
-	union
-	{
-		FMemberId MemberIds[3] = {};
-		struct { FMemberId Add, Del, Set; } Ops;
-	};
+	enum class EMember : uint8 { Del, Add };
+	FMemberId MemberIds[2];
 	
 	template<class Ids>	
 	void InitIds()
 	{
-		static FSetDeltaOps Cache = {.MemberIds = {Ids::IndexMember("Add"), Ids::IndexMember("Del"), Ids::IndexMember("Set")}};
+		static FSetDeltaOps Cache = {.MemberIds = {Ids::IndexMember("Del"), Ids::IndexMember("Add")}};
 		*this = Cache;
 	}
 };
 
-
 template <typename T, typename KeyFuncs, typename SetAllocator>
 struct TSetDeltaBinding : ICustomBinding, FSetDeltaOps
 {
 	using Type = TSet<T, KeyFuncs, SetAllocator>;
 	static constexpr EMemberPresence Occupancy = EMemberPresence::AllowSparse;
 
+	struct FCustomTypename
+	{
+		inline static constexpr std::string_view DeclName = "SetDelta";
+		inline static constexpr std::string_view BindName = Concat<DeclName, ShortTypename<KeyFuncs>, ShortTypename<SetAllocator>>;
+		inline static constexpr std::string_view Namespace = "PlainProps::UE::";
+		using Parameters = std::tuple<T>;
+	};
+
 	void Save(FMemberBuilder& Dst, const Type& Src, const Type* Default, const FSaveContext& Context) const
 	{
-		if (Default)
+		if (!Default || Default->IsEmpty())
 		{
-			if (Default->IsEmpty())
-			{
-				// Todo: Add everything
-			}
-			else
-			{
-				// TODO: Range builder for missing items, iterate over defaults elements and check existence in Src
-				// Dst.AddRange(Ops.Del, MissingItems) if non-empty;
+			// Todo: Add everything
 
-				// TODO: Range builder for added items, iterate over defaults elements and check existence in Src
-			}
 		}
 		else
 		{
-			// Dst.AddRange(Ops.Set, );
+			// TODO: Range builder for missing items, iterate over defaults elements and check existence in Src
+			// Dst.AddRange(Ops.Del, MissingItems) if non-empty;
+
+			// TODO: Range builder for added items, iterate over defaults elements and check existence in Src
 		}
 	}
 
@@ -509,23 +507,19 @@ struct TSetDeltaBinding : ICustomBinding, FSetDeltaOps
 		FMemberId Name = Members.PeekName().Get();
 		FRangeView Items = Members.GrabRange();
 		int32 NumItems = static_cast<int32>(Items.Num());
-		if (Name == Ops.Set)
-		{
-			Dst.Empty(NumItems);
-			AddItems(Dst, Items, Batch);
-		}
-		else if (Name == Ops.Add)
+		if (Name == MemberIds[(uint8)EMember::Add])
 		{
 			Dst.Reserve(Dst.Num() + NumItems);
 			AddItems(Dst, Items, Batch);
 		}
-		else if (Name == Ops.Del)
+		else
 		{
+			check(Members.PeekName() == MemberIds[(uint8)EMember::Del]);
 			RemoveItems(Dst, Items);
 		
 			if (Members.HasMore())
 			{
-				check(Members.PeekName() == Ops.Add);
+				check(Members.PeekName() == MemberIds[(uint8)EMember::Add]);
 				Items = Members.GrabRange();
 				Dst.Reserve(Dst.Num() + static_cast<int32>(Items.Num()));
 				AddItems(Dst, Items, Batch);
diff --git a/Engine/Plugins/Experimental/PlainProps/Source/Tests/PlainPropsUeCoreTests.cpp b/Engine/Plugins/Experimental/PlainProps/Source/Tests/PlainPropsUeCoreTests.cpp
--- a/Engine/Plugins/Experimental/PlainProps/Source/Tests/PlainPropsUeCoreTests.cpp
+++ b/Engine/Plugins/Experimental/PlainProps/Source/Tests/PlainPropsUeCoreTests.cpp
@@ -125,7 +125,7 @@ struct TScopedStructBinding : TScopedStructDeclaration<T, Occupancy, Runtime>
 	FStructSchemaId BindId;
 
 	TScopedStructBinding()
-	: BindId(IndexStructBindIdIfNeeded<T, Ids>(Super::DeclId))
+	: BindId(IndexStructBindIdIfNeeded<Typename, Ids>(Super::DeclId))
 	{
 		BindNativeStruct<CttiOf<T>, Runtime>(Runtime::GetSchemas(), BindId, Super::DeclId);
 	}
@@ -1007,7 +1007,7 @@ TEST_CASE_NAMED(FPlainPropsUeCoreTest, "System::Core::Serialization::PlainProps:
 	SECTION("TSetDelta")
 	{
 		TScopedStructBinding<FInt> Int;
-		//TScopedStructBinding<FSets, EMemberPresence::AllowSparse, FDeltaRuntime> Sets;
+		TScopedStructBinding<FSets, EMemberPresence::AllowSparse, FDeltaRuntime> Sets;
 	}
 
 	SECTION("Transform")
