PlainProps TSet delta bindings + various bug fixes

#rnx

diff --git a/Engine/Plugins/Experimental/PlainProps/Source/Private/PlainPropsBind.cpp b/Engine/Plugins/Experimental/PlainProps/Source/Private/PlainPropsBind.cpp
--- a/Engine/Plugins/Experimental/PlainProps/Source/Private/PlainPropsBind.cpp
+++ b/Engine/Plugins/Experimental/PlainProps/Source/Private/PlainPropsBind.cpp
@@ -167,6 +167,12 @@ void FCustomBindings::DropStruct(FStructSchemaId BindId)
 	checkf(false, TEXT("'%s' unbound"), *Debug.Print(BindId));
 }
 
+FOptionalStructSchemaId	FCustomBindings::FindStructDeclId(FStructSchemaId BindId) const
+{ 
+	 FCustomBindingEntry Entry = Find(BindId);
+	 return Entry ? ToOptional(Entry.DeclId) : NoId;
+}
+
 FCustomBindingEntry	FCustomBindings::Find(FStructSchemaId BindId) const
 {
 	for (FCustomBindingEntry Entry : Entries)
@@ -264,6 +270,14 @@ FStructSchemaId FSchemaBindings::GetDeclId(FStructSchemaId BindId) const
 
 //////////////////////////////////////////////////////////////////////////
 
+FStructSchemaId FStructBindIds::GetDeclId(FStructSchemaId BindId) const
+{
+	FOptionalStructSchemaId CustomId = Customs.FindStructDeclId(BindId);
+	return CustomId ? CustomId.Get() : Schemas.GetDeclId(BindId);
+}
+
+//////////////////////////////////////////////////////////////////////////
+
 TArray<FStructSchemaId> IndexInMemoryNames(const FSchemaBatch& Schemas, FIdIndexerBase& Indexer)
 {
 	const uint8* Base = reinterpret_cast<const uint8*>(&Schemas);
diff --git a/Engine/Plugins/Experimental/PlainProps/Source/Private/PlainPropsBuild.cpp b/Engine/Plugins/Experimental/PlainProps/Source/Private/PlainPropsBuild.cpp
--- a/Engine/Plugins/Experimental/PlainProps/Source/Private/PlainPropsBuild.cpp
+++ b/Engine/Plugins/Experimental/PlainProps/Source/Private/PlainPropsBuild.cpp
@@ -128,12 +128,12 @@ void FMemberBuilder::BuildSuperStruct(FScratchAllocator& Scratch, const FStructD
 		return;
 	}
 	
-	FBuiltStructPtr OnlyMember = BuildAndReset(Scratch, Super, Debug);
+	FBuiltStruct* OnlyMember = BuildAndReset(Scratch, Super, Debug);
 	Members.Emplace(FBuiltMember::MakeSuper(Super.Id, MoveTemp(OnlyMember)));
 	check(Members[0].Schema.Type.AsStruct().IsSuper);
 }
 
-FBuiltStructPtr FMemberBuilder::BuildAndReset(FScratchAllocator& Scratch, const FStructDeclaration& Declared, const FDebugIds& Debug)
+FBuiltStruct* FMemberBuilder::BuildAndReset(FScratchAllocator& Scratch, const FStructDeclaration& Declared, const FDebugIds& Debug)
 {
 	checkf(!(Declared.Super && Declared.Occupancy == EMemberPresence::RequireAll),
 		TEXT("Requiring sub structs to be dense isn't implemented"));
@@ -164,12 +164,12 @@ FBuiltStructPtr FMemberBuilder::BuildAndReset(FScratchAllocator& Scratch, const
 
 	Members.Reset();
 
-	return FBuiltStructPtr(Out);
+	return Out;
 }
 
 //////////////////////////////////////////////////////////////////////////
 
-FBuiltStructPtr FDenseMemberBuilder::BuildHomo(const FStructDeclaration& Declaration, FMemberType Leaf, TConstArrayView<FBuiltValue> Values) const
+FBuiltStruct* FDenseMemberBuilder::BuildHomo(const FStructDeclaration& Declaration, FMemberType Leaf, TConstArrayView<FBuiltValue> Values) const
 {
 	check(Declaration.NumMembers == Values.Num());
 
@@ -224,11 +224,11 @@ FBuiltMember::FBuiltMember(FMemberId Name, FTypedRange Range)
 : FBuiltMember(Name, MoveTemp(Range.Schema), { .Range = Range.Values })
 {}
 
-FBuiltMember::FBuiltMember(FMemberId Name, FStructSchemaId Schema, FBuiltStructPtr Value)
+FBuiltMember::FBuiltMember(FMemberId Name, FStructSchemaId Schema, FBuiltStruct* Value)
 : FBuiltMember(Name, MakeMemberSchema(DefaultStructType, Schema), { .Struct = Value })
 {}
 
-FBuiltMember FBuiltMember::MakeSuper(FStructSchemaId Schema, FBuiltStructPtr Value)
+FBuiltMember FBuiltMember::MakeSuper(FStructSchemaId Schema, FBuiltStruct* Value)
 {
 	return FBuiltMember(NoId, MakeMemberSchema(SuperStructType, Schema), { .Struct = Value });
 }
diff --git a/Engine/Plugins/Experimental/PlainProps/Source/Private/PlainPropsLoad.cpp b/Engine/Plugins/Experimental/PlainProps/Source/Private/PlainPropsLoad.cpp
--- a/Engine/Plugins/Experimental/PlainProps/Source/Private/PlainPropsLoad.cpp
+++ b/Engine/Plugins/Experimental/PlainProps/Source/Private/PlainPropsLoad.cpp
@@ -37,13 +37,15 @@ public:
 
 	// @param OffsetWidth Usage unimplemented, store size and offsets as 8/16/32/64-bit 
 	explicit FLoadStructPlan(const FSchemaBinding& Schema, ELeafWidth OffsetWidth, bool bSparse)
-	: Handle(uint64(&Schema) | (uint64(OffsetWidth) << 1) | (uint64(bSparse) << FPlatformMemory::KernelAddressBit) | SchemaMask)
+	: Handle(uint64(&Schema) | SchemaBit | (bSparse ? SparseBit : 0) | (uint64(OffsetWidth) << SchemaOffsetShift))
 	{
+		static_assert(alignof(FSchemaBinding) >= 8);
 		check(&Schema == &AsSchema());
-		check(IsSparseSchema() == bSparse);
+		check(bSparse == IsSparseSchema());
+		check(OffsetWidth == GetOffsetWidth());
 	}
 
-	bool						IsSchema() const		{ return (Handle & SchemaMask) == SchemaMask; }
+	bool						IsSchema() const		{ return (Handle & SchemaBit) == SchemaBit; }
 	bool						IsSparseSchema() const	{ return (Handle & SparseSchemaMask) == SparseSchemaMask; }
 	bool						IsMemcpy() const		{ return (Handle & LoMask) == MemcpyMask; }
 	bool						IsCustom() const		{ return (Handle & LoMask) == CustomMask; }
@@ -51,14 +53,18 @@ public:
 	const ICustomBinding&		AsCustom() const		{ check(IsCustom()); return *AsPtr<ICustomBinding>(); }
 	const FSchemaBinding&		AsSchema() const		{ check(IsSchema()); return *AsPtr<FSchemaBinding>(); }
 
+	ELeafWidth					GetOffsetWidth() const	{ check(IsSchema()); return static_cast<ELeafWidth>((Handle & SchemaOffsetMask) >> SchemaOffsetShift); }
+
 private:
-	static constexpr uint64 SparseMask			= uint64(1) << FPlatformMemory::KernelAddressBit;
-	static constexpr uint64 PtrMask				= ~(SparseMask | 0b111);
+	static constexpr uint64 SparseBit			= uint64(1) << FPlatformMemory::KernelAddressBit;
+	static constexpr uint64 PtrMask				= ~(SparseBit | 0b111);
 	static constexpr uint64 LoMask				= 0b11;
 	static constexpr uint64 MemcpyMask			= 0b00;
 	static constexpr uint64 CustomMask			= 0b10;
-	static constexpr uint64 SchemaMask			= 0b01;
-	static constexpr uint64 SparseSchemaMask	= SchemaMask | SparseMask;
+	static constexpr uint64 SchemaBit			= 0b01;
+	static constexpr uint64 SparseSchemaMask	= SchemaBit | SparseBit;
+	static constexpr uint64 SchemaOffsetShift	= 1;
+	static constexpr uint64 SchemaOffsetMask	= 0b110;
 	
 	template<typename T>
 	const T* AsPtr() const
@@ -128,7 +134,7 @@ static bool HasDifferentSupers(const FStructSchema& From, const FSchemaBinding&
 // Used to create an additional load plan, beyond the saved struct schema ids
 struct FLoadIdMapping
 {
-	FStructSchemaId SaveId; // ~ Batch decl id, index into saved schemas and load plans
+	FStructSchemaId ReadId; // ~ Batch decl id, index into saved schemas and load plans
 	FStructSchemaId LoadId; // ~ Batch bind id, index into load plans
 	FStructSchemaId BindId; // Runtime bind id
 };
@@ -142,10 +148,10 @@ public:
 	, NextLoadIdx(static_cast<uint32>(RuntimeDeclIds.Num()))
 	{}
 
-	FStructSchemaId					BindLoadId(FStructSchemaId SaveId, FStructSchemaId BindId)
+	FStructSchemaId					BindLoadId(FStructSchemaId ReadId, FStructSchemaId BindId)
 	{
-		FStructSchemaId DeclId = DeclIds[SaveId.Idx];
-		return BindId == DeclId ? SaveId : MapLoadId(SaveId, BindId);
+		FStructSchemaId DeclId = DeclIds[ReadId.Idx];
+		return BindId == DeclId ? ReadId : MapLoadId(ReadId, BindId);
 	}
 	FLoadIdMapping					GetMapping(int32 Idx) const { return Mappings[Idx]; }
 	int32							NumMappings() const { return Mappings.Num(); }
@@ -155,19 +161,19 @@ private:
 	uint32								NextLoadIdx;
 	TArray<FLoadIdMapping>				Mappings;
 
-	FStructSchemaId MapLoadId(FStructSchemaId SaveId, FStructSchemaId BindId)
+	FStructSchemaId MapLoadId(FStructSchemaId ReadId, FStructSchemaId BindId)
 	{
 		for (FLoadIdMapping Mapping : Mappings)
 		{
 			if (Mapping.BindId == BindId)
 			{
-				check(Mapping.SaveId == SaveId);
+				check(Mapping.ReadId == ReadId);
 				return Mapping.LoadId;
 			}
 		}
 
 		FLoadIdMapping& Mapping = Mappings.AddDefaulted_GetRef();
-		Mapping.SaveId = SaveId;
+		Mapping.ReadId = ReadId;
 		Mapping.LoadId = { NextLoadIdx++ };
 		Mapping.BindId = BindId;
 		return Mapping.LoadId;
@@ -179,10 +185,23 @@ private:
 struct FLoadBatch
 {
 	FReadBatchId			ReadId; // Needed to access schemas for custom struct loading
+	uint32					NumReadSchemas;
 	uint32					NumPlans;
 	FLoadStructPlan			Plans[0];
 
-	FLoadStructPlan			operator[](FStructSchemaId Id) const { check(Id.Idx < NumPlans); return Plans[Id.Idx]; }
+	FLoadStructPlan			operator[](FStructSchemaId LoadId) const
+	{
+		check(LoadId.Idx < NumPlans);
+		return Plans[LoadId.Idx];
+	}
+
+	FStructSchemaId			GetReadId(FStructSchemaId LoadId) const
+	{
+		check(LoadId.Idx < NumPlans);
+		static_assert(alignof(FLoadStructPlan) >= alignof(FStructSchemaId));
+		const FStructSchemaId* SaveIds = reinterpret_cast<const FStructSchemaId*>(Plans + NumPlans);
+		return LoadId.Idx < NumReadSchemas ? LoadId : SaveIds[LoadId.Idx - NumReadSchemas];
+	}
 };
 
 void FLoadBatchDeleter::operator()(FLoadBatch* Batch) const
@@ -239,12 +258,16 @@ static void CreateSubsetBindingWithoutEnumIds(const FStructSchema& From, const F
 {
 	check(To.NumMembers == ToNames.Num());
 	check(To.NumMembers >= From.NumMembers);
+	check(IsAligned(Out.Num(), alignof(FSchemaBinding)));
 
 	int32 OutPos = Out.Num();
+	
+	// Allocate and init header 
 	FSchemaBinding Header = { To.DeclId, From.NumMembers, From.NumInnerSchemas - NumEnums, From.NumRangeTypes };
 	Out.AddUninitialized(Header.CalculateSize());
 	FSchemaBinding* Schema = new (&Out[OutPos]) FSchemaBinding {Header};
 	
+	// Copy subset of member bindings
 	FMemberVisitor ToIt(To);
 	FMemberBinder Footer(*Schema);
 	const FSchemaId* InnerSchemaIt = From.GetInnerSchemas();
@@ -262,11 +285,15 @@ static void CreateSubsetBindingWithoutEnumIds(const FStructSchema& From, const F
 
 static void CloneBindingWithReplacedStructIds(const FSchemaId* FromIds, const FSchemaBinding& To, FLoadIdBinder& LoadIds, SubsetByteArray& Out)
 {
+	check(IsAligned(Out.Num(), alignof(FSchemaBinding)));
+
+	// Clone runtime FSchemaBinding including inner schema bind ids
 	uint32 Size = To.CalculateSize();
 	Out.AddUninitialized(Size);
 	FSchemaBinding* Schema = reinterpret_cast<FSchemaBinding*>(&Out[Out.Num() - Size]);
 	FMemory::Memcpy(Schema, &To, Size);
 
+	// Replace inner schema bind ids with batch load ids
 	const FStructSchemaId* SaveIdIt = static_cast<const FStructSchemaId*>(FromIds);
 	FStructSchemaId* OutIds = static_cast<FStructSchemaId*>(const_cast<FSchemaId*>(Schema->GetInnerSchemas()));
 	for (FStructSchemaId& OutId : MakeArrayView(OutIds, To.NumInnerSchemas))
@@ -350,7 +377,7 @@ struct FLoadPlanner
 		for (uint32 Idx = 0; Idx < NumPlans; ++Idx)
 		{
 			FLoadIdMapping Mapping;
-			Mapping.SaveId = { Idx };
+			Mapping.ReadId = { Idx };
 			Mapping.LoadId = { Idx };
 			Mapping.BindId = RuntimeIds[Idx];
 			CreatePlan(Mapping, LoadIds);
@@ -372,7 +399,7 @@ struct FLoadPlanner
 			// Verify that all UnboundSaveIds were bound by some load plan
 			for (int32 Idx = 0; !UnboundSaveIds.IsEmpty() && Idx < LoadIds.NumMappings(); ++Idx)
 			{
-				UnboundSaveIds.Remove(LoadIds.GetMapping(Idx).SaveId);
+				UnboundSaveIds.Remove(LoadIds.GetMapping(Idx).ReadId);
 			}
 		}
 
@@ -382,25 +409,41 @@ struct FLoadPlanner
 		}
 
 		// Allocate load batch, copy plans and subset schemas, and fixup subset schema plans
-		return CreateBatch();
+		return CreateBatch(LoadIds);
 	}
 private:
 
-	FLoadBatchPtr CreateBatch()
+	FLoadBatchPtr CreateBatch(const FLoadIdBinder& LoadIds)
 	{
 		const uint32 NumPlans = Plans.Num();
-
-		SIZE_T Bytes = sizeof(FLoadBatch) + sizeof(FLoadStructPlan) * NumPlans + SubsetSchemaData.Num();
-		FLoadBatch Header = { ReadId, NumPlans };
+		const uint32 NumMappings = LoadIds.NumMappings();
+		const uint32 NumReadSchemas = RuntimeIds.Num();
+		check(NumPlans == NumReadSchemas + NumMappings);
+
+		// Allocate plan, init header and copy plans
+		SIZE_T Bytes = sizeof(FLoadBatch) + sizeof(FLoadStructPlan) * NumPlans
+					 + Align(sizeof(FStructSchemaId) * NumMappings, alignof(FSchemaBinding))
+					 + SubsetSchemaData.Num();
+		FLoadBatch Header = { ReadId, NumReadSchemas, NumPlans };
 		FLoadBatch* Out = new (FMemory::Malloc(Bytes)) FLoadBatch{Header};
 		FMemory::Memcpy(Out->Plans, Plans.GetData(), sizeof(FLoadStructPlan) * NumPlans);
+
+		// Copy LoadId -> ReadId mapping so custom-bound mapped plans can form FReadSchemaHandle and FStructView
+		FStructSchemaId* OutReadId = reinterpret_cast<FStructSchemaId*>(Out->Plans + NumPlans);
+		for (uint32 Idx = 0; Idx < NumMappings; ++Idx)
+		{
+			*OutReadId++ = LoadIds.GetMapping(Idx).ReadId;
+			check(OutReadId[-1].Idx < NumReadSchemas);
+		}
+
+		// Copy cloned subset schemas and patch up plan -> schema pointers
 		if (SubsetSchemaData.Num() > 0)
 		{
-			uint8* OutSubsetData = reinterpret_cast<uint8*>(Out->Plans + NumPlans);
+			void* OutSubsetData = Align(OutReadId, alignof(FSchemaBinding));
 			FMemory::Memcpy(OutSubsetData, SubsetSchemaData.GetData(), SubsetSchemaData.Num());
 		
 			// Update plans with actual subset schema pointers
-			const uint8* It = OutSubsetData;
+			const uint8* It = static_cast<const uint8*>(OutSubsetData);
 			for (uint32 Idx = 0; Idx < NumPlans; ++Idx)
 			{
 				if (int32 Size = SubsetSchemaSizes[Idx])
@@ -412,7 +455,8 @@ private:
 					It += Size;
 				}
 			}
-			check(It == OutSubsetData + SubsetSchemaData.Num());
+			check(It == (uint8*)OutSubsetData + SubsetSchemaData.Num());
+			check(It == (uint8*)Out + Bytes);
 		}
 
 		return FLoadBatchPtr(Out);
@@ -432,7 +476,7 @@ private:
 			return FLoadStructPlan(*Custom);
 		}
 
-		const FStructSchema& From = ResolveStructSchema(ReadId, Mapping.SaveId);
+		const FStructSchema& From = ResolveStructSchema(ReadId, Mapping.ReadId);
 		if (From.NumMembers)
 		{
 			if (const FSchemaBinding* To = Schemas.FindStruct(Mapping.BindId))
@@ -443,12 +487,14 @@ private:
 			}
 			
 			// Type-erased structs
-			UnboundSaveIds.Add(Mapping.SaveId);
+			UnboundSaveIds.Add(Mapping.ReadId);
 			return FLoadStructPlan(GUnboundBinding);
 		}
-		
-		checkf(!From.Type.Scope && From.Type.Name.NumParameters == 2, TEXT("Only range-bound template parameters are memberless. "
-				"They're always anonymous two-parameter types and uninstantiable as structs, bound via MakeAnonymousParametricType()"));
+
+//		checkf(!From.Type.Scope && From.Type.Name.NumParameters == 2, TEXT("Unexpected bind id typename '%s' saved as '%s'. Only range-bound template parameters are memberless. "
+//				"They're always anonymous two-parameter types and uninstantiable as structs, bound via MakeAnonymousParametricType(). "),
+//				*Declarations.GetDebug().Print(From.Type),
+//				*Declarations.GetDebug().Print(Mapping.SaveId));
 		return FLoadStructPlan(GMemberlessBinding);
 	}
 };
@@ -504,9 +550,12 @@ public:
 	{
 		TArray<FMemberBindType, TFixedAllocator<16>> InnerTypes;
 		InnerTypes.Add(ToBindType(Src.Schema.ItemType));
-		for (const FMemberType* It = Src.Schema.NestedItemTypes; It; It = It->IsRange() ? (It + 1) : nullptr)
+		if (Src.Schema.ItemType.IsRange())
 		{
-			InnerTypes.Add(ToBindType(*It));
+			for (const FMemberType* It = Src.Schema.NestedItemTypes; It; It = It->IsRange() ? (It + 1) : nullptr)
+			{
+				InnerTypes.Add(ToBindType(*It));
+			}
 		}
 		check(Bindings.Num() == InnerTypes.Num());
 		
@@ -516,13 +565,12 @@ public:
 		
 		FByteReader ByteIt(Src.Values);
 		FBitCacheReader BitIt;
-		LoadRange(Member, ByteIt, BitIt, Batch, Plan);
+		LoadRangePlan(Member, Src.NumItems, ByteIt, BitIt, Batch, Plan);
 		ByteIt.CheckEmpty();
 	}
 
-	static void LoadRange(uint8* Member, FByteReader& ByteIt, FBitCacheReader& BitIt, const FLoadBatch& Batch, const FLoadRangePlan& Range)
+	static void LoadRangePlan(uint8* Member, uint64 Num, FByteReader& ByteIt, FBitCacheReader& BitIt, const FLoadBatch& Batch, const FLoadRangePlan& Range)
 	{
-		uint64 Num = GrabRangeNum(Range.MaxSize, ByteIt, BitIt);
 		FRangeBinding Binding = Range.Bindings[0];
 		FMemberBindType InnerType = Range.InnerTypes[0];
 		
@@ -618,7 +666,8 @@ public:
 		uint64 ItemSize = Items.Size;
 		for (uint8* It = Items.Data, *End = It + Items.NumBytes(); It != End; It += ItemSize)
 		{
-			LoadRange(It, ByteIt, BitIt, Batch, Plan);	
+			uint64 Num = GrabRangeNum(Plan.MaxSize, ByteIt, BitIt);
+			LoadRangePlan(It, Num, ByteIt, BitIt, Batch, Plan);
 		}
 	}
 };
@@ -733,17 +782,18 @@ private:
 		PlainProps::LoadStruct(Member, FByteReader(ByteIt.GrabSkippableSlice()), Id, Batch);
 	}
 
-	void LoadMemberRange(uint8* Member, FLoadRangePlan&& Range)
+	void LoadMemberRange(uint8* Member, const FLoadRangePlan& Plan)
 	{
-		LoadRange(Member, ByteIt, BitIt, Batch, Range);
+		uint64 Num = GrabRangeNum(Plan.MaxSize, ByteIt, BitIt);
+		LoadRangePlan(Member, Num, ByteIt, BitIt, Batch, Plan);
 	}
 };
 
 ////////////////////////////////////////////////////////////////////////////
 
-void LoadStruct(void* Dst, FByteReader Src, FStructSchemaId Id, const FLoadBatch& Batch)
+void LoadStruct(void* Dst, FByteReader Src, FStructSchemaId LoadId, const FLoadBatch& Batch)
 {
-	FLoadStructPlan Plan = Batch[Id];
+	FLoadStructPlan Plan = Batch[LoadId];
 	if (Plan.IsSchema())
 	{
 		if (Plan.IsSparseSchema())
@@ -762,7 +812,7 @@ void LoadStruct(void* Dst, FByteReader Src, FStructSchemaId Id, const FLoadBatch
 	}
 	else
 	{
-		FStructSchemaHandle ReadSchema{Id, Batch.ReadId};
+		FStructSchemaHandle ReadSchema{Batch.GetReadId(LoadId), Batch.ReadId};
 		Plan.AsCustom().LoadCustom(Dst, { ReadSchema, Src }, ECustomLoadMethod::Assign, Batch);
 	}
 }
diff --git a/Engine/Plugins/Experimental/PlainProps/Source/Private/PlainPropsSave.cpp b/Engine/Plugins/Experimental/PlainProps/Source/Private/PlainPropsSave.cpp
--- a/Engine/Plugins/Experimental/PlainProps/Source/Private/PlainPropsSave.cpp
+++ b/Engine/Plugins/Experimental/PlainProps/Source/Private/PlainPropsSave.cpp
@@ -194,7 +194,7 @@ template<class SaverType, typename InnerContextType>
 	return { Member.InnerTypes + 1, Member.RangeBindings + 1, static_cast<uint16>(Member.NumRanges - 1), Member.InnermostSchema };
 }
 
-[[nodiscard]] static FBuiltRange* SaveLeafRangeBinding(FScratchAllocator& Scratch, const uint8* Range, const ILeafRangeBinding& Binding, FUnpackedLeafType Leaf)
+[[nodiscard]] static FBuiltRange* SaveLeafRangeBinding(FScratchAllocator& Scratch, const void* Range, const ILeafRangeBinding& Binding, FUnpackedLeafType Leaf)
 {
 	FLeafRangeAllocator Allocator(Scratch, Leaf);
 	Binding.SaveLeaves(Range, Allocator);
@@ -219,7 +219,7 @@ ELeafWidth GetArithmeticWidth(FLeafBindType Leaf)
 	return Leaf.Arithmetic.Width;
 }
 
-[[nodiscard]] static FBuiltRange* SaveRange(const uint8* Range, FRangeMemberBinding Member, const FSaveContext& Ctx)
+[[nodiscard]] FBuiltRange* SaveRange(const void* Range, FRangeMemberBinding Member, const FSaveContext& Ctx)
 {
 	FRangeBinding Binding = Member.RangeBindings[0];
 	FMemberBindType InnerType = Member.InnerTypes[0];
diff --git a/Engine/Plugins/Experimental/PlainProps/Source/Public/PlainPropsBind.h b/Engine/Plugins/Experimental/PlainProps/Source/Public/PlainPropsBind.h
--- a/Engine/Plugins/Experimental/PlainProps/Source/Public/PlainPropsBind.h
+++ b/Engine/Plugins/Experimental/PlainProps/Source/Public/PlainPropsBind.h
@@ -126,7 +126,7 @@ static_assert(sizeof(FMemberBindType) == 1);
 ////////////////////////////////////////////////////////////////////////////////////////////////
 
 // Members are loaded in saved FStructSchema order, not current offset order unless upgrade layer reorders
-struct FSchemaBinding
+struct alignas(/*FRangeBinding*/ 8) FSchemaBinding
 {
 	FStructSchemaId			DeclId;
 	uint16					NumMembers;
@@ -284,7 +284,7 @@ public:
 	// @param Binding must outlive this or call DropStruct()
 	PLAINPROPS_API void						BindStruct(FStructSchemaId BindId, FStructSchemaId DeclId, ICustomBinding& Binding);
 	const ICustomBinding*					FindStruct(FStructSchemaId BindId) const		{ return Find(BindId).Binding; }
-	FStructSchemaId							FindStructDeclId(FStructSchemaId BindId) const	{ return Find(BindId).DeclId; }
+	FOptionalStructSchemaId					FindStructDeclId(FStructSchemaId BindId) const;
 	FCustomBindingEntry						FindStructToSave(FStructSchemaId BindId)		{ return Find(BindId); }
 	PLAINPROPS_API void						DropStruct(FStructSchemaId BindId);
 
@@ -560,11 +560,25 @@ public:
 	PLAINPROPS_API const FSchemaBinding&	GetStruct(FStructSchemaId BindId) const;
 	PLAINPROPS_API void						DropStruct(FStructSchemaId BindId);
 
+	virtual FStructSchemaId					GetDeclId(FStructSchemaId BindId) const override final;
 private:
 	TArray<TUniquePtr<FSchemaBinding>>		Bindings;
 	const FDebugIds&						Debug;
+};
+
+////////////////////////////////////////////////////////////////////////////////////////////////
 
-	virtual FStructSchemaId GetDeclId(FStructSchemaId BindId) const override;
+struct FStructBindIds : IStructBindIds
+{
+	FStructBindIds(const FCustomBindings& InCustoms, const FSchemaBindings& InSchemas)
+	: Customs(InCustoms)
+	, Schemas(InSchemas)
+	{}
+
+	const FCustomBindings& Customs;
+	const FSchemaBindings& Schemas;
+
+	virtual FStructSchemaId GetDeclId(FStructSchemaId BindId) const override final;
 };
 
 ////////////////////////////////////////////////////////////////////////////////////////////////
@@ -1039,6 +1053,67 @@ struct FMemberBinder
 
 //////////////////////////////////////////////////////////////////////////
 
+template<typename T>
+inline constexpr FMemberType ReflectInnermostType()
+{
+	// DefaultStructType
+	return FMemberType(FStructType{EMemberKind::Struct, /* IsDynamic */ 0, /* IsSuper */ 0});;
+}
+
+template<Arithmetic T>
+inline constexpr FMemberType ReflectInnermostType()
+{
+	return ReflectArithmetic<T>.Pack();
+}
+template<Enumeration T>
+inline constexpr FMemberType ReflectInnermostType()
+{
+	return ReflectEnum<T>.Pack();
+}
+
+// Hack / workaround related to C++'s lack of templated nullary constructors
+union FUninitializedMemberBindType
+{
+	FUninitializedMemberBindType() : Unused() {}
+	uint8_t			Unused;
+	FMemberBindType Value;
+};
+
+// Helps templated custom bindings save containers as ranges
+template<class RangeBinding>
+struct TRangeMemberHelper
+{
+	static constexpr uint16 NumRanges = CountRangeBindings<RangeBinding>();
+	static constexpr ERangeSizeType MaxSize = RangeSizeOf(typename RangeBinding::SizeType{});
+	using InnermostType = typename TInnerType<RangeBinding, NumRanges>::Type;
+
+	const FRangeBinding*			RangeBindings = nullptr;
+	FOptionalSchemaId				InnermostSchema;
+	FUninitializedMemberBindType	InnerBindTypes[NumRanges];
+	FMemberType						InnerSchemaTypes[NumRanges];
+
+	template<class Ids>
+	void Init()
+	{
+		RangeBindings = GetRangeBindings<RangeBinding, NumRanges>().GetData();
+		for (uint16 Idx = 0; Idx < NumRanges - 1; ++Idx)
+		{
+			ERangeSizeType Type = RangeBindings[Idx + 1].GetSizeType();
+			InnerBindTypes[Idx].Value = FMemberBindType(Type);
+			InnerSchemaTypes[Idx] = FMemberType(Type);
+		}
+		InnerBindTypes[NumRanges - 1].Value = BindInnermostType<InnermostType, Ids>(/* out */ InnermostSchema);
+		InnerSchemaTypes[NumRanges - 1] = ReflectInnermostType<InnermostType>();
+	}
+
+	FRangeMemberBinding MakeBinding(uint32 Offset) const
+	{
+		return { &InnerBindTypes[0].Value, RangeBindings, NumRanges, InnermostSchema, Offset };
+	}
+};
+
+//////////////////////////////////////////////////////////////////////////
+
 // Save -> load struct ids for ESchemaFormat::InMemoryNames
 [[nodiscard]] PLAINPROPS_API TArray<FStructSchemaId> IndexInMemoryNames(const FSchemaBatch& Schemas,  FIdIndexerBase& Indexer);
 
diff --git a/Engine/Plugins/Experimental/PlainProps/Source/Public/PlainPropsBuild.h b/Engine/Plugins/Experimental/PlainProps/Source/Public/PlainPropsBuild.h
--- a/Engine/Plugins/Experimental/PlainProps/Source/Public/PlainPropsBuild.h
+++ b/Engine/Plugins/Experimental/PlainProps/Source/Public/PlainPropsBuild.h
@@ -19,8 +19,6 @@ struct FBuiltRange;
 class FDebugIds;
 struct FUnpackedLeafType;
 
-using FBuiltStructPtr = FBuiltStruct*;
-
 //////////////////////////////////////////////////////////////////////////
 
 /// Single-threaded scratch allocator for intermediate built representation
@@ -164,6 +162,13 @@ inline FMemberSchema MakeStructRangeSchema(ERangeSizeType SizeType, FStructSchem
 
 PLAINPROPS_API FMemberSchema MakeNestedRangeSchema(FScratchAllocator& Scratch, ERangeSizeType SizeType, FMemberSchema InnerRangeSchema);
 
+// @param InnerTypes must outlive FMemberSchema
+template<uint16 N>
+inline FMemberSchema MakeNestedRangeSchema(ERangeSizeType SizeType, const FMemberType (&InnerTypes)[N], FOptionalSchemaId InnermostSchema)
+{
+	return { FMemberType(SizeType), InnerTypes[0], N, InnermostSchema, N > 1 ? InnerTypes : nullptr };
+}
+
 //////////////////////////////////////////////////////////////////////////
 
 [[nodiscard]] PLAINPROPS_API FBuiltRange* CloneLeaves(FScratchAllocator& Scratch, uint64 Num, const void* Data, SIZE_T LeafSize);
@@ -206,9 +211,9 @@ struct FBuiltMember
 {
 	FBuiltMember(FMemberId Name, FUnpackedLeafType Leaf, FOptionalEnumSchemaId Schema, uint64 Value);
 	FBuiltMember(FMemberId Name, FTypedRange Range);
-	FBuiltMember(FMemberId Name, FStructSchemaId Schema, FBuiltStructPtr Value);
+	FBuiltMember(FMemberId Name, FStructSchemaId Schema, FBuiltStruct* Value);
 	FBuiltMember(FOptionalMemberId N, FMemberSchema S, FBuiltValue V) : Name(N), Schema(MoveTemp(S)), Value(V) {}
-	static FBuiltMember MakeSuper(FStructSchemaId Schema, FBuiltStructPtr Value);
+	static FBuiltMember MakeSuper(FStructSchemaId Schema, FBuiltStruct* Value);
 
 	FOptionalMemberId		Name;
 	FMemberSchema			Schema;
@@ -240,12 +245,12 @@ public:
 
 	void AddLeaf(FMemberId Name, FUnpackedLeafType Leaf, FOptionalEnumSchemaId Enum, uint64 Value)	{ Members.Emplace(Name, Leaf, Enum, Value); }
 	void AddRange(FMemberId Name, FTypedRange Range)												{ Members.Emplace(Name, Range); }
-	void AddStruct(FMemberId Name, FStructSchemaId Schema, FBuiltStructPtr Struct)					{ Members.Emplace(Name, Schema, Struct); }
+	void AddStruct(FMemberId Name, FStructSchemaId Schema, FBuiltStruct* Struct)					{ Members.Emplace(Name, Schema, Struct); }
 	
 	// Build members into a single nested super struct member, no-op if no non-super members has been added
 	PLAINPROPS_API void BuildSuperStruct(FScratchAllocator&	Scratch, const FStructDeclaration& Super, const FDebugIds& Debug);
 
-	[[nodiscard]] PLAINPROPS_API FBuiltStructPtr BuildAndReset(FScratchAllocator& Scratch, const FStructDeclaration& Declared, const FDebugIds& Debug);
+	[[nodiscard]] PLAINPROPS_API FBuiltStruct* BuildAndReset(FScratchAllocator& Scratch, const FStructDeclaration& Declared, const FDebugIds& Debug);
 
 	bool IsEmpty() const { return Members.IsEmpty(); }
 
@@ -267,7 +272,7 @@ struct FDenseMemberBuilder
 	const FDebugIds& Debug;
 
 	template<typename T, typename... Ts>
-	[[nodiscard]] FBuiltStructPtr BuildHomogeneous(const FStructDeclaration& Declaration, T Head, Ts... Tail) const
+	[[nodiscard]] FBuiltStruct* BuildHomogeneous(const FStructDeclaration& Declaration, T Head, Ts... Tail) const
 	{
 		// Todo: Handle enums, ranges and structs
 		FBuiltValue Values[] = { {.Leaf = ValueCast(Head)}, {.Leaf = (ValueCast(Tail))}...  };
@@ -275,7 +280,7 @@ struct FDenseMemberBuilder
 	}
 
 private:
-	[[nodiscard]] PLAINPROPS_API FBuiltStructPtr BuildHomo(const FStructDeclaration& Declaration, FMemberType Leaf, TConstArrayView<FBuiltValue> Values) const;
+	[[nodiscard]] PLAINPROPS_API FBuiltStruct* BuildHomo(const FStructDeclaration& Declaration, FMemberType Leaf, TConstArrayView<FBuiltValue> Values) const;
 };
 
 // Helper class for building struct ranges
diff --git a/Engine/Plugins/Experimental/PlainProps/Source/Public/PlainPropsSave.h b/Engine/Plugins/Experimental/PlainProps/Source/Public/PlainPropsSave.h
--- a/Engine/Plugins/Experimental/PlainProps/Source/Public/PlainPropsSave.h
+++ b/Engine/Plugins/Experimental/PlainProps/Source/Public/PlainPropsSave.h
@@ -3,15 +3,18 @@
 #pragma once
 
 #include "PlainPropsTypes.h"
-#include "PlainPropsBuild.h"
 #include "Templates/UniquePtr.h"
 
 namespace PlainProps 
 {
 
+struct FBuiltRange;
+struct FBuiltStruct;
 class FCustomBindings;
 class FDeclarations;
+struct FRangeMemberBinding;
 class FSchemaBindings;
+class FScratchAllocator;
 
 // Temporary data structure, will be replaced by something more sophisticated
 // perhaps deduplicating all zero-memory defaults
@@ -33,7 +36,7 @@ FSaveContext MakeSaveContext(FDefaultStructs Defaults, FScratchAllocator& Scratc
 	return { Runtime::GetTypes(), Runtime::GetSchemas(), Runtime::GetCustoms(), Scratch, Defaults };
 }
 
-[[nodiscard]] PLAINPROPS_API FBuiltStructPtr SaveStruct(const void* Struct, FStructSchemaId BindId, const FSaveContext& Context);
-[[nodiscard]] PLAINPROPS_API FBuiltStructPtr SaveStructDelta(const void* Struct, const void* Default, FStructSchemaId BindId, const FSaveContext& Context);
-
+[[nodiscard]] PLAINPROPS_API FBuiltStruct*	SaveStruct(const void* Struct, FStructSchemaId BindId, const FSaveContext& Context);
+[[nodiscard]] PLAINPROPS_API FBuiltStruct*	SaveStructDelta(const void* Struct, const void* Default, FStructSchemaId BindId, const FSaveContext& Context);
+[[nodiscard]] PLAINPROPS_API FBuiltRange*	SaveRange(const void* Range, FRangeMemberBinding Member, const FSaveContext& Ctx);
 } // namespace PlainProps
\ No newline at end of file
diff --git a/Engine/Plugins/Experimental/PlainProps/Source/Public/PlainPropsUeCoreBindings.h b/Engine/Plugins/Experimental/PlainProps/Source/Public/PlainPropsUeCoreBindings.h
--- a/Engine/Plugins/Experimental/PlainProps/Source/Public/PlainPropsUeCoreBindings.h
+++ b/Engine/Plugins/Experimental/PlainProps/Source/Public/PlainPropsUeCoreBindings.h
@@ -8,9 +8,10 @@
 #include "Math/Transform.h"
 #include "Misc/Optional.h"
 #include "PlainPropsBind.h"
+#include "PlainPropsBuild.h" // FMemberSchema for TSetDeltaBinding::SaveSet
+#include "PlainPropsIndex.h"
 #include "PlainPropsLoad.h"
 #include "PlainPropsRead.h"
-#include "PlainPropsIndex.h"
 #include "PlainPropsStringUtil.h"
 #include "UObject/NameTypes.h"
 
@@ -24,6 +25,12 @@ PP_REFLECT_STRUCT(, FVector4, void, X, Y, Z, W);
 PP_REFLECT_STRUCT(, FQuat, void, X, Y, Z, W);
 }
 
+template <typename T>
+struct TIsContiguousContainer<PlainProps::TRangeView<T>>
+{
+	static inline constexpr bool Value = true;
+};
+
 namespace PlainProps::UE
 {
 
@@ -419,7 +426,7 @@ struct TMapBinding : public TSetBinding<TPair<K, V>, KeyFuncs, SetAllocator>
 
 //////////////////////////////////////////////////////////////////////////
 
-//TODO: macroify, e.g PP_CUSTOM_BIND(PLAINPROPS_API, FTransform, Transform, Translate, Rotate, Scale)
+//TODO: Consider macroifying parts of this, e.g PP_CUSTOM_BIND(PLAINPROPS_API, FTransform, Transform, Translate, Rotate, Scale)
 struct FTransformBinding : ICustomBinding
 {
 	using Type = FTransform;
@@ -464,147 +471,214 @@ template <typename T, typename KeyFuncs, typename SetAllocator>
 struct TSetDeltaBinding : ICustomBinding, FSetDeltaOps
 {
 	using Type = TSet<T, KeyFuncs, SetAllocator>;
+	using FSet = Type;
+	using FSetBinding = TSetBinding<T, KeyFuncs, SetAllocator>;
+	using FRangeMemberHelper = TRangeMemberHelper<FSetBinding>;
 	static constexpr EMemberPresence Occupancy = EMemberPresence::AllowSparse;
+	static constexpr uint16 NumInnerRanges = FRangeMemberHelper::NumRanges - 1;
 
 	struct FCustomTypename
 	{
 		inline static constexpr std::string_view DeclName = "SetDelta";
 		inline static constexpr std::string_view BindName = Concat<DeclName, ShortTypename<KeyFuncs>, ShortTypename<SetAllocator>>;
-		inline static constexpr std::string_view Namespace = "PlainProps::UE::";
+		inline static constexpr std::string_view Namespace;
 		using Parameters = std::tuple<T>;
 	};
 
-	void Save(FMemberBuilder& Dst, const Type& Src, const Type* Default, const FSaveContext& Context) const
+	FRangeMemberHelper InnerRange;
+	
+	template<class Ids>	
+	void InitIds()
+	{
+		FSetDeltaOps::InitIds<Ids>();
+		InnerRange.template Init<Ids>();
+	}
+	
+	void Save(FMemberBuilder& Dst, const FSet& Src, const FSet* Default, const FSaveContext& Context) const
 	{
-		if (!Default || Default->IsEmpty())
+		if (Default && !Default->IsEmpty())
 		{
-			// Todo: Add everything
-
+			// Inefficient, mimic FSetProperty::SerializeItem for better perf
+			SaveSet(Dst, MemberIds[(uint8)EMember::Del], Default->Difference(Src), Context);
+			SaveSet(Dst, MemberIds[(uint8)EMember::Add], Src.Difference(*Default), Context);
 		}
-		else
+		else if (Src.Num())
 		{
-			// TODO: Range builder for missing items, iterate over defaults elements and check existence in Src
-			// Dst.AddRange(Ops.Del, MissingItems) if non-empty;
+			SaveSet(Dst, MemberIds[(uint8)EMember::Add], Src, Context);
+		}
+	}
 
-			// TODO: Range builder for added items, iterate over defaults elements and check existence in Src
+	void SaveSet(FMemberBuilder& Dst, FMemberId Name, const FSet& Set, const FSaveContext& Ctx) const
+	{
+		if (int32 Num = Set.Num())
+		{
+			// Start of a more optimized version that saves typed values directly w/o leaning on TSetBinding
+			//
+			//FBuiltRange* Range;
+			//if constexpr (LeafType<T>)
+			//{
+			//	Range = FBuiltRange::Create(Ctx.Scratch, Num, sizeof(T));
+			//	T* OutIt = static_cast<T*>(Range->Data);
+			//	//const TSparseArray<TSetElement<T>& Elements = reinterpret_cast<const TSparseArray<TSetElement<T>>&>(Values);
+			//	//if (Elements.IsCompact())
+			//	//{
+			//	//	FMemory::Memcpy(Range->Data, Elements.GetData(), sizeof(T) * Num);
+			//	//}
+			//	//else 
+			//	for (T Leaf : Set)
+			//	{
+			//		*OutIt++ = Leaf;
+			//	}
+			//}
+			//else 
+			//{
+			//	Range = SaveRange(&Set, InnerRange.MakeBinding(/* offset*/ 0), Ctx);
+			//}
+
+			// Lean on TSetBinding for now, less efficient but simpler
+			FBuiltRange* Range = SaveRange(&Set, InnerRange.MakeBinding(/* offset*/ 0), Ctx);
+			FMemberSchema Schema = MakeNestedRangeSchema(InnerRange.MaxSize, InnerRange.InnerSchemaTypes, InnerRange.InnermostSchema);
+			Dst.AddRange(Name, { Schema, Range });	
 		}
+
 	}
 
-	inline void Load(Type& Dst, FStructView Src, ECustomLoadMethod Method, const FLoadBatch& Batch) const
+	inline void Load(FSet& Dst, FStructView Src, ECustomLoadMethod Method, const FLoadBatch& Batch) const
 	{
 		FMemberReader Members(Src);
 
 		if (Method == ECustomLoadMethod::Construct)
 		{
-			::new (&Dst) Type;
+			::new (&Dst) FSet;
 		}
 				
-		if (!Members.HasMore())
+		if (Members.HasMore())
 		{
-			return;
+			if (Members.PeekName() == MemberIds[(uint8)EMember::Add])
+			{
+				ApplyItems<EMember::Add>(Dst, Members.GrabRange(), Batch);
+			}
+			else
+			{
+				check(Members.PeekName() == MemberIds[(uint8)EMember::Del]);
+				ApplyItems<EMember::Del>(Dst, Members.GrabRange(), Batch);
+		
+				if (Members.HasMore())
+				{
+					check(Members.PeekName() == MemberIds[(uint8)EMember::Add]);
+					ApplyItems<EMember::Add>(Dst, Members.GrabRange(), Batch);
+				}
+			}
+
+			check(!Members.HasMore());
 		}
+	}
+
+	template<EMember Op>
+	void ApplyItems(FSet& Out, FRangeView Items, const FLoadBatch& Batch) const
+	{
+		check(!Items.IsEmpty());
 
-		FMemberId Name = Members.PeekName().Get();
-		FRangeView Items = Members.GrabRange();
-		int32 NumItems = static_cast<int32>(Items.Num());
-		if (Name == MemberIds[(uint8)EMember::Add])
+		if constexpr (Op == EMember::Add && !LeafType<T>)
 		{
-			Dst.Reserve(Dst.Num() + NumItems);
-			AddItems(Dst, Items, Batch);
+			Out.Reserve(static_cast<int32>(Items.Num()));
 		}
-		else
+
+		if constexpr (LeafType<T>)
 		{
-			check(Members.PeekName() == MemberIds[(uint8)EMember::Del]);
-			RemoveItems(Dst, Items);
-		
-			if (Members.HasMore())
-			{
-				check(Members.PeekName() == MemberIds[(uint8)EMember::Add]);
-				Items = Members.GrabRange();
-				Dst.Reserve(Dst.Num() + static_cast<int32>(Items.Num()));
-				AddItems(Dst, Items, Batch);
-			}
+			ApplyLeaves<Op>(Out, Items.AsLeaves().As<T>());
+		}
+		else if constexpr (NumInnerRanges)
+		{
+			ApplyRanges<Op>(Out, Items.AsRanges(), Batch);
+		}
+		else 
+		{
+			ApplyStructs<Op>(Out, Items.AsStructs(), Batch);
 		}
-		
-		check(!Members.HasMore());
 	}
-
-	inline static bool Diff(const Type& A, const Type& B)
+	
+	template<EMember Op>
+	void ApplyLeaves(FSet& Out, TRangeView<T> Items) const
 	{
-		if (A.Num() != B.Num())
+		if constexpr (Op == EMember::Add)
 		{
-			return false;
+			Out.Append(MakeArrayView(Items));
 		}
-
-		for (const T& AKey : A)
+		else
 		{
-			if (!B.Contains(AKey))
+			for (T Item : Items)
 			{
-				return false;
+				Out.Remove(Item);
 			}
 		}
-
-		return true;
 	}
 
-	static void AddItems(Type& Out, FRangeView Items, const FLoadBatch& Batch)
+	template<EMember Op>
+	void ApplyRanges(FSet& Out, FNestedRangeView Items, const FLoadBatch& Batch) const
 	{
-		check(!Items.IsEmpty());
+		static_assert(std::is_default_constructible_v<T>, TEXT("Ranges must be default-constructible"));
 
-		if constexpr (LeafType<T>)
+		TConstArrayView<FRangeBinding> InnerRangeBindings(InnerRange.RangeBindings + 1, NumInnerRanges);
+		for (FRangeView Item : Items)
 		{
-			for (T Item : Items.AsLeaves().As<T>())
-			{
-				Out.Add(Item);
-			}
+			T Tmp;
+			LoadRange(&Tmp, Item, InnerRange.MaxSize, InnerRangeBindings, Batch);
+			ApplyItem<Op>(Out, MoveTemp(Tmp));
 		}
-		else if (Items.IsStructRange())
-		{
-			FStructRangeView Structs = Items.AsStructs();
-			for (FStructView Item : Structs)
-			{
-				if constexpr (std::is_default_constructible_v<T>)
-				{
-					T Tmp;
-					PlainProps::LoadStruct(&Tmp, Item, Batch);	
-					Out.Emplace(MoveTemp(Tmp));
-				}
-				else
-				{
-					alignas(T) uint8 Buffer[sizeof(T)];
-					PlainProps::ConstructAndLoadStruct(Buffer, Item, Batch);
-					T* Tmp = reinterpret_cast<T*>(Buffer);
-					Out.Emplace(MoveTemp(*Tmp));
-					Tmp->~T();
-				}
-			}
-		}
-		else // Nested range
+	}
+	
+	template<EMember Op>
+	void ApplyStructs(FSet& Out, FStructRangeView Items, const FLoadBatch& Batch) const
+	{
+		for (FStructView Item : Items)
 		{
-			using RangeBinding = RangeBind<T>;
-			if constexpr (std::is_default_constructible_v<T> && !std::is_void_v<RangeBinding>)
+			if constexpr (std::is_default_constructible_v<T>)
 			{
-				static constexpr ERangeSizeType MaxSize = RangeSizeOf(typename RangeBinding::SizeType{});
-				//const IItemRangeBinding* Bindings[] = {};// TODO ... generate somehow ... ;
 				T Tmp;
-				for (FRangeView Item : Items.AsRanges())
-				{
-				//	LoadRange(&Tmp, Item, MaxSize, Bindings, Batch);
-				//	Out.Emplace(MoveTemp(*Tmp));
-				}
+				LoadStruct(&Tmp, Item, Batch);
+				ApplyItem<Op>(Out, MoveTemp(Tmp));
 			}
 			else
 			{
-				check(Items.IsNestedRange());
-				checkf(std::is_default_constructible_v<T>, TEXT("Ranges must be default-constructible"));
-				checkf(!std::is_void_v<RangeBinding>, TEXT("Inner range type unbound (no TRangeBind specialization)"));
+				alignas(T) uint8 Buffer[sizeof(T)];
+				ConstructAndLoadStruct(Buffer, Item, Batch);
+				T& Tmp = *reinterpret_cast<T*>(Buffer);
+				ApplyItem<Op>(Out, MoveTemp(Tmp));
+				Tmp.~T();
 			}
 		}
 	}
+	
+	template<EMember Op>
+	void ApplyItem(FSet& Out, T&& Item) const
+	{
+		if constexpr (Op == EMember::Add)
+		{
+			Out.Emplace(MoveTemp(Item));
+		}
+		else
+		{
+			Out.Remove(Item);
+		}
+	}
 
-	static void RemoveItems(Type& Out, FRangeView Items)
+	inline static bool Diff(const FSet& A, const FSet& B)
 	{
-		// TODO copy paste from AddItems or template AddItems -> ApplyItems<EOp::Add/Del>
+		if (A.Num() != B.Num())
+		{
+			return true;
+		}
+
+		for (const T& AKey : A)
+		{
+			if (!B.Contains(AKey))
+			{
+				return true;
+			}
+		}
+
+		return false;
 	}
 };
 
@@ -665,7 +739,7 @@ struct TSetDeltaBinding : ICustomBinding, FSetDeltaOps
 //		}
 //	}
 //
-//	virtual FBuiltStructPtr	SaveStruct(const void* Src, const FDebugIds& Debug) const override
+//	virtual FBuiltStruct*	SaveStruct(const void* Src, const FDebugIds& Debug) const override
 //	{
 //		...
 //	}
diff --git a/Engine/Plugins/Experimental/PlainProps/Source/Tests/PlainPropsTests.cpp b/Engine/Plugins/Experimental/PlainProps/Source/Tests/PlainPropsTests.cpp
--- a/Engine/Plugins/Experimental/PlainProps/Source/Tests/PlainPropsTests.cpp
+++ b/Engine/Plugins/Experimental/PlainProps/Source/Tests/PlainPropsTests.cpp
@@ -123,7 +123,7 @@ public:
 	TArray64<uint8>				Write();
 
 private:
-	TArray<TPair<FStructSchemaId, FBuiltStructPtr>>	Objects;
+	TArray<TPair<FStructSchemaId, FBuiltStruct*>>	Objects;
 	FDeclarations									Declarations;
 	FScratchAllocator&								Scratch;
 
@@ -152,7 +152,7 @@ private:
 
 	TArray<char> GetNameData() const;
 
-	virtual FStructSchemaId GetDeclId(FStructSchemaId BindId) const override
+	virtual FStructSchemaId GetDeclId(FStructSchemaId BindId) const override final
 	{
 		checkf(false, TEXT("All struct ids should be declared, nothing is bound with different names in this test suite"));
 		return BindId;
@@ -192,7 +192,7 @@ TArray64<uint8> FTestBatchBuilder::Write()
 {
 	// Build partial schemas
 	FSchemasBuilder SchemaBuilders(Declarations, *this, Scratch);
-	for (const TPair<FStructSchemaId, FBuiltStructPtr>& Object : Objects)
+	for (const TPair<FStructSchemaId, FBuiltStruct*>& Object : Objects)
 	{
 		SchemaBuilders.NoteStructAndMembers(Object.Key, *Object.Value);
 	}
@@ -226,7 +226,7 @@ TArray64<uint8> FTestBatchBuilder::Write()
 
 	// Write objects
 	WriteU32(Out, TestMagics[2]);
-	for (const TPair<FStructSchemaId, FBuiltStructPtr>& Object : Objects)
+	for (const TPair<FStructSchemaId, FBuiltStruct*>& Object : Objects)
 	{
 		WriteU32(/* out */ Tmp, TestMagics[3]);
 		WriteU32(/* out */ Tmp, Writer.GetWriteId(Object.Key).Get().Idx);
@@ -636,14 +636,14 @@ TEST_CASE_NAMED(FPlainPropsReadWriteTest, "System::Core::Serialization::PlainPro
 
 			FMemberBuilder Members;
 			Members.Add(	Batch.NameMember("I1"), 100);
-			FBuiltStructPtr NestedInStruct = Members.BuildAndReset(Scratch, Batch.Get(NestedId), Batch);
+			FBuiltStruct* NestedInStruct = Members.BuildAndReset(Scratch, Batch.Get(NestedId), Batch);
 
 			Members.AddStruct(Batch.NameMember("Nested"), NestedId, MoveTemp(NestedInStruct));
 			Members.Add(Batch.NameMember("Leaf"), true);
-			FBuiltStructPtr Struct = Members.BuildAndReset(Scratch, Batch.Get(StructId), Batch);
+			FBuiltStruct* Struct = Members.BuildAndReset(Scratch, Batch.Get(StructId), Batch);
 		
 			Members.Add(Batch.NameMember("I2"), 200);
-			FBuiltStructPtr NestedInObject = Members.BuildAndReset(Scratch, Batch.Get(NestedId), Batch);
+			FBuiltStruct* NestedInObject = Members.BuildAndReset(Scratch, Batch.Get(NestedId), Batch);
 
 			Members.Add(Batch.NameMember("L1"), 123.f);
 			Members.AddStruct(Batch.NameMember("S"), StructId, MoveTemp(Struct));
diff --git a/Engine/Plugins/Experimental/PlainProps/Source/Tests/PlainPropsUeCoreTests.cpp b/Engine/Plugins/Experimental/PlainProps/Source/Tests/PlainPropsUeCoreTests.cpp
--- a/Engine/Plugins/Experimental/PlainProps/Source/Tests/PlainPropsUeCoreTests.cpp
+++ b/Engine/Plugins/Experimental/PlainProps/Source/Tests/PlainPropsUeCoreTests.cpp
@@ -194,7 +194,7 @@ struct FNameBinding : FTestCustomBinding
 class FBatchSaver
 {
 public:
-	FBatchSaver();
+	explicit FBatchSaver(const FCustomBindings& CustomBase);
 
 	template<class T>
 	void						Save(T&& Object);
@@ -204,15 +204,15 @@ public:
 	TArray64<uint8>				Write() const;
 
 private:
-	using IdBuiltStructPair = TPair<FStructSchemaId, FBuiltStructPtr>;
+	using IdBuiltStructPair = TPair<FStructSchemaId, FBuiltStruct*>;
 	TArray<IdBuiltStructPair>	SavedObjects;
 	FNameBinding				SavedNames;
 	FCustomBindings				Customs;
 	mutable FScratchAllocator	Scratch;
 };
 
-FBatchSaver::FBatchSaver()
-: Customs(/* debug */ GNames, &GCustoms)
+FBatchSaver::FBatchSaver(const FCustomBindings& CustomBase)
+: Customs(/* debug */ GNames, &CustomBase)
 {
 	Customs.BindStruct(SavedNames.Declaration.Id, SavedNames.Declaration.Id, SavedNames);
 }
@@ -228,7 +228,7 @@ template<class T>
 bool FBatchSaver::SaveDelta(const T& Object, const T& Default) 
 {
 	FDualStructSchemaId Id = IndexStructDualId<std::remove_reference_t<T>, FIds>();
-	if (FBuiltStructPtr Delta = SaveStructDelta(&Object, &Default, Id.Bind, {GTypes, GSchemas, Customs, Scratch}))
+	if (FBuiltStruct* Delta = SaveStructDelta(&Object, &Default, Id.Bind, {GTypes, GSchemas, Customs, Scratch}))
 	{
 		SavedObjects.Emplace(Id.Decl, MoveTemp(Delta));
 		return true;
@@ -255,7 +255,8 @@ inline constexpr uint32 Magics[] = { 0xFEEDF00D, 0xABCD1234, 0xDADADAAA, 0x99887
 TArray64<uint8> FBatchSaver::Write() const
 {
 	// Build partial schemas
-	FSchemasBuilder SchemaBuilders(GTypes, GSchemas, Scratch);
+	const FStructBindIds BindIds(Customs, GSchemas);
+	FSchemasBuilder SchemaBuilders(GTypes, BindIds, Scratch);
 	for (const IdBuiltStructPair& Object : SavedObjects)
 	{
 		SchemaBuilders.NoteStructAndMembers(Object.Key, *Object.Value);
@@ -263,7 +264,7 @@ TArray64<uint8> FBatchSaver::Write() const
 	FBuiltSchemas Schemas = SchemaBuilders.Build(); 
 
 	// Filter out declared but unused names and ids
-	FWriter Writer(GNames, GSchemas, Schemas, ESchemaFormat::StableNames);
+	FWriter Writer(GNames, BindIds, Schemas, ESchemaFormat::StableNames);
 	TArray<FName> UsedNames;
 	for (uint32 Idx = 0, Num = GNames.NumNames(); Idx < Num; ++Idx)
 	{
@@ -288,7 +289,7 @@ TArray64<uint8> FBatchSaver::Write() const
 
 	// Write objects
 	WriteU32(Out, Magics[2]);
-	for (const TPair<FStructSchemaId, FBuiltStructPtr>& Object : SavedObjects)
+	for (const TPair<FStructSchemaId, FBuiltStruct*>& Object : SavedObjects)
 	{
 		WriteU32(/* out */ Tmp, Magics[3]);
 		WriteU32(/* out */ Tmp, Writer.GetWriteId(Object.Key).Get().Idx);
@@ -313,8 +314,8 @@ TArray64<uint8> FBatchSaver::Write() const
 class FBatchLoader
 {
 public:
-	FBatchLoader(FMemoryView Data)
-	: Customs(/* debug */ GNames, &GCustoms)
+	FBatchLoader(FMemoryView Data, const FCustomBindings& CustomBase)
+	: Customs(/* debug */ GNames, &CustomBase)
 	{
 		// Read ids
 		FByteReader It(Data);
@@ -371,6 +372,14 @@ public:
 		LoadInto(Out);
 		return MoveTemp(Out);
 	}
+	
+	template<class T>
+	T Load(const T& Original)
+	{
+		T Out = Original;
+		LoadInto(Out);
+		return MoveTemp(Out);
+	}
 
 	template<class T>
 	void LoadInto(T& Out)
@@ -389,16 +398,16 @@ private:
 };
 
 
-static void Run(void (*Save)(FBatchSaver&), void (*Load)(FBatchLoader&))
+static void Run(void (*Save)(FBatchSaver&), void (*Load)(FBatchLoader&), const FCustomBindings& Customs = GCustoms)
 {
 	TArray64<uint8> Data;
 	{
-		FBatchSaver Batch;
+		FBatchSaver Batch(Customs);
 		Save(Batch);
 		Data = Batch.Write();
 	}
 
-	FBatchLoader Batch(MakeMemoryView(Data));
+	FBatchLoader Batch(MakeMemoryView(Data), Customs);
 	Load(Batch);
 }
 
@@ -513,12 +522,31 @@ struct FSets
 	TSet<char> Leaves;
 	TSet<TArray<uint8>> Ranges;
 	TSet<FInt> Structs;
+	TSet<FString> Strings; // ILeafRangeBinding
 };
-PP_REFLECT_STRUCT(PlainProps::UE::Test, FSets, void, Leaves, Ranges, Structs);
+PP_REFLECT_STRUCT(PlainProps::UE::Test, FSets, void, Leaves, Ranges, Structs, Strings);
+
+template<typename T>
+inline bool OrderEq(const TSet<T>& A, const TSet<T>& B)
+{
+	if (A.Num() != B.Num())
+	{
+		return false;
+	}
+
+	for (typename TSet<T>::TConstIterator AIt(A), BIt(B); AIt; ++AIt, ++BIt)
+	{
+		if (*AIt != *BIt)
+		{
+			return false;
+		}
+	}
+	return true;
+}
 
 inline bool operator==(const FSets& A, const FSets& B)
 {
-	return LegacyCompareEqual(A.Leaves, B.Leaves) && LegacyCompareEqual(A.Ranges, B.Ranges) && LegacyCompareEqual(A.Structs, B.Structs);
+	return OrderEq(A.Leaves, B.Leaves) && OrderEq(A.Ranges, B.Ranges) && OrderEq(A.Structs, B.Structs) && OrderEq(A.Strings, B.Strings);
 }
 
 struct FMaps
@@ -853,9 +881,7 @@ TEST_CASE_NAMED(FPlainPropsUeCoreTest, "System::Core::Serialization::PlainProps:
 			{
 				CHECK(Batch.Load<FOpts>() == FOpts{});
 				CHECK(Batch.Load<FOpts>() == FOpts{{true}, {FNDC{2}}, {FNDCIntrusive{3}}});
-				FOpts AlreadySet = {{false}, {FNDC{0}}, {FNDCIntrusive{1}}};
-				Batch.LoadInto(AlreadySet);
-				CHECK(AlreadySet == FOpts{{true}, {FNDC{2}}, {FNDCIntrusive{3}}});
+				CHECK(Batch.Load(FOpts{{false}, {FNDC{0}}, {FNDCIntrusive{1}}}) == FOpts{{true}, {FNDC{2}}, {FNDCIntrusive{3}}});
 			});
 	}
 
@@ -992,10 +1018,8 @@ TEST_CASE_NAMED(FPlainPropsUeCoreTest, "System::Core::Serialization::PlainProps:
 			}, 
 			[](FBatchLoader& Batch)
 			{
-				FDelta Zero = {false, 0, {}, {}, {}};
-				FDelta DefaultOnZero = Zero;
-				Batch.LoadInto(DefaultOnZero);
-				CHECK(DefaultOnZero == FDelta{});
+				const FDelta Zero = {false, 0, {}, {}, {}};
+				CHECK(Batch.Load(Zero) == FDelta{});
 				CHECK(Batch.Load<FDelta>() == Zero);
 				CHECK(Batch.Load<FDelta>() == FDelta{.B = 123});
 				CHECK(Batch.Load<FDelta>() == FDelta{.C = {321}});
@@ -1008,6 +1032,28 @@ TEST_CASE_NAMED(FPlainPropsUeCoreTest, "System::Core::Serialization::PlainProps:
 	{
 		TScopedStructBinding<FInt> Int;
 		TScopedStructBinding<FSets, EMemberPresence::AllowSparse, FDeltaRuntime> Sets;
+		Run([](FBatchSaver& Batch)
+			{
+				Batch.Save(FSets{});
+				Batch.Save(FSets{{'l'}, {{1}}, {{2}}, {"s"}});
+				const FSets Default = {{'a'}, {{1}}, {{1}}, {"a"}};
+				CHECK(!Batch.SaveDelta(FSets{}, FSets{}));
+				CHECK(Batch.SaveDelta(Default, FSets{}));
+				CHECK(Batch.SaveDelta(FSets{}, Default)); // Wipe
+				CHECK(!Batch.SaveDelta(Default, Default));
+				CHECK(Batch.SaveDelta(FSets{{'a'}, {{0,1,2}}, {{2}}, {}}, Default)); // Mixed changes
+			}, 
+			[](FBatchLoader& Batch)
+			{
+				CHECK(Batch.Load<FSets>() == FSets{});
+				CHECK(Batch.Load<FSets>() == FSets{{'l'}, {{1}}, {{2}}, {"s"}});
+				
+				const FSets Default = {{'a'}, {{1}}, {{1}}, {"a"}};
+				CHECK(Batch.Load<FSets>() == Default);
+				CHECK(Batch.Load(Default) == FSets{}); // Wipe
+				CHECK(Batch.Load(Default) == FSets{{'a'}, {{0,1,2}}, {{2}}, {}}); // Mixed changes
+			},
+			GDeltaCustoms);
 	}
 
 	SECTION("Transform")
