PlainProps add TCustomInit

#rnx

diff --git a/Engine/Plugins/Experimental/PlainProps/Source/Public/PlainPropsBind.h b/Engine/Plugins/Experimental/PlainProps/Source/Public/PlainPropsBind.h
--- a/Engine/Plugins/Experimental/PlainProps/Source/Public/PlainPropsBind.h
+++ b/Engine/Plugins/Experimental/PlainProps/Source/Public/PlainPropsBind.h
@@ -243,8 +243,8 @@ protected: // for unit tests
 
 struct FDualStructSchemaId
 {
-	FStructSchemaId Bind;
-	FStructSchemaId Decl;
+	FStructSchemaId BindId;
+	FStructSchemaId DeclId;
 };
 
 ////////////////////////////////////////////////////////////////////////////////////////////////
@@ -627,7 +627,7 @@ FTypeId IndexStructName()
 	}
 }
 
-template<typename Typename, class Ids>
+template<class Ids, typename Typename>
 FStructSchemaId IndexStructBindIdIfNeeded(FStructSchemaId DeclId)
 {
 	if constexpr (ExplicitBindName<Typename> || ParametricName<Typename>)
@@ -641,20 +641,17 @@ FStructSchemaId IndexStructBindIdIfNeeded(FStructSchemaId DeclId)
 	}
 }
 
-template<typename Struct, class Ids>
-FDualStructSchemaId IndexStructDualId()
-{
-	using Typename = TTypename<Struct>;
-	
+template<class Ids, typename Typename>
+FDualStructSchemaId IndexStructDualId(FTypeId DeclName = IndexStructName<Ids, ETypename::Decl, Typename>())
+{	
 	FDualStructSchemaId Out;
-	FTypeId DeclName = IndexStructName<Ids, ETypename::Decl, Typename>();
-	Out.Decl = Ids::IndexStruct(DeclName);
-	Out.Bind = Out.Decl;
+	Out.DeclId = Ids::IndexStruct(DeclName);
+	Out.BindId = Out.DeclId;
 
 	if constexpr (ExplicitBindName<Typename> || ParametricName<Typename>)
 	{
 		FTypeId BindName = IndexStructName<Ids, ETypename::Bind, Typename>();
-		Out.Bind = BindName != DeclName ? Ids::IndexStruct(BindName) : Out.Decl;
+		Out.BindId = BindName != DeclName ? Ids::IndexStruct(BindName) : Out.DeclId;
 	}
 
 	return Out;
@@ -751,22 +748,23 @@ FTypeId IndexParametricType(FTypeId TemplatedType, const std::tuple<Ts...>*)
 
 ////////////////////////////////////////////////////////////////////////////////////////////////
 
+template<class Ids>
+struct TCustomInit {};
+
+
 template<typename CustomBinding, class Runtime>
 FStructSchemaId BindCustomStructOnce()
 {
-	struct FBinding : CustomBinding
+	struct FBinding : FDualStructSchemaId, CustomBinding
 	{
 		using Type = typename CustomBinding::Type;
 		using Typename = TCustomTypename<CustomBinding>;
 		using Ids = typename Runtime::Ids;
 
-		FBinding()
+		explicit FBinding(FTypeId DeclType = IndexStructName<Ids, ETypename::Decl, Typename>())
+		: FDualStructSchemaId(IndexStructDualId<Ids, Typename>(DeclType))
+		, CustomBinding(TCustomInit<Ids>{})
 		{
-			FTypeId DeclType = IndexStructName<Ids, ETypename::Decl, Typename>();
-			DeclId = Ids::IndexStruct(DeclType);
-			BindId = IndexStructBindIdIfNeeded<Typename, Ids>(DeclId);
-			
-			CustomBinding::template InitIds<Ids>();
 			Runtime::GetTypes().DeclareStruct(DeclId, DeclType, CustomBinding::MemberIds, CustomBinding::Occupancy);
 			Runtime::GetCustoms().BindStruct(BindId, DeclId, *this);
 		}
@@ -791,9 +789,6 @@ FStructSchemaId BindCustomStructOnce()
 		{
 			return CustomBinding::Diff(*static_cast<const Type*>(A), *static_cast<const Type*>(B));
 		}
-		
-		FStructSchemaId DeclId;
-		FStructSchemaId BindId;
 	};
 
 	static FBinding Binding;
@@ -1071,7 +1066,7 @@ inline constexpr FMemberType ReflectInnermostType()
 	return ReflectEnum<T>.Pack();
 }
 
-// Hack / workaround related to C++'s lack of templated nullary constructors
+// TRangeMemberHelper init helper
 union FUninitializedMemberBindType
 {
 	FUninitializedMemberBindType() : Unused() {}
diff --git a/Engine/Plugins/Experimental/PlainProps/Source/Public/PlainPropsUeCoreBindings.h b/Engine/Plugins/Experimental/PlainProps/Source/Public/PlainPropsUeCoreBindings.h
--- a/Engine/Plugins/Experimental/PlainProps/Source/Public/PlainPropsUeCoreBindings.h
+++ b/Engine/Plugins/Experimental/PlainProps/Source/Public/PlainPropsUeCoreBindings.h
@@ -432,20 +432,16 @@ struct FTransformBinding : ICustomBinding
 	using Type = FTransform;
 	inline static constexpr EMemberPresence Occupancy = EMemberPresence::AllowSparse;
 	enum class EMember : uint8 { Translate, Rotate, Scale };
-	FMemberId MemberIds[3];
-	FStructSchemaId VectorId;
-	FStructSchemaId QuatId;
+	const FMemberId MemberIds[3];
+	const FStructSchemaId VectorId;
+	const FStructSchemaId QuatId;
 
 	template<class Ids>
-	void InitIds(/*const FDeclarations& Declared*/) 
-	{
-		MemberIds[(uint8)EMember::Translate] = Ids::IndexMember("Translate");
-		MemberIds[(uint8)EMember::Rotate] = Ids::IndexMember("Rotate");
-		MemberIds[(uint8)EMember::Scale] = Ids::IndexMember("Scale");
-
-		VectorId = GetStructDeclId<Ids, FVector>();
-		QuatId = GetStructDeclId<Ids, FQuat>();
-	}
+	FTransformBinding(TCustomInit<Ids>/*, const FDeclarations& Declared*/)
+	: MemberIds{Ids::IndexMember("Translate"), Ids::IndexMember("Rotate"), Ids::IndexMember("Scale")}
+	, VectorId(GetStructDeclId<Ids, FVector>())
+	, QuatId(GetStructDeclId<Ids, FQuat>())
+	{}
 
 	PLAINPROPS_API void	Save(FMemberBuilder& Dst, const FTransform& Src, const FTransform* Default, const FSaveContext& Context) const;
 	PLAINPROPS_API void	Load(FTransform& Dst, FStructView Src, ECustomLoadMethod Method, const FLoadBatch& Batch) const;
@@ -454,21 +450,26 @@ struct FTransformBinding : ICustomBinding
 
 //////////////////////////////////////////////////////////////////////////
 
-struct FSetDeltaOps
+struct FSetDeltaIds
 {
 	enum class EMember : uint8 { Del, Add };
-	FMemberId MemberIds[2];
+	const FMemberId MemberIds[2];
 	
 	template<class Ids>	
-	void InitIds()
+	FSetDeltaIds(TCustomInit<Ids>)
+	: MemberIds{Ids::IndexMember("Del"), Ids::IndexMember("Add")}
+	{}
+
+	template<class Ids>	
+	static FSetDeltaIds Cache()
 	{
-		static FSetDeltaOps Cache = {.MemberIds = {Ids::IndexMember("Del"), Ids::IndexMember("Add")}};
-		*this = Cache;
+		static FSetDeltaIds Out(TCustomInit<Ids>{});
+		return Out;
 	}
 };
 
 template <typename T, typename KeyFuncs, typename SetAllocator>
-struct TSetDeltaBinding : ICustomBinding, FSetDeltaOps
+struct TSetDeltaBinding : ICustomBinding, FSetDeltaIds
 {
 	using Type = TSet<T, KeyFuncs, SetAllocator>;
 	using FSet = Type;
@@ -488,9 +489,9 @@ struct TSetDeltaBinding : ICustomBinding, FSetDeltaOps
 	FRangeMemberHelper InnerRange;
 	
 	template<class Ids>	
-	void InitIds()
+	TSetDeltaBinding(TCustomInit<Ids>)
+	: FSetDeltaIds(FSetDeltaIds::Cache<Ids>())
 	{
-		FSetDeltaOps::InitIds<Ids>();
 		InnerRange.template Init<Ids>();
 	}
 	
diff --git a/Engine/Plugins/Experimental/PlainProps/Source/Tests/PlainPropsUeCoreTests.cpp b/Engine/Plugins/Experimental/PlainProps/Source/Tests/PlainPropsUeCoreTests.cpp
--- a/Engine/Plugins/Experimental/PlainProps/Source/Tests/PlainPropsUeCoreTests.cpp
+++ b/Engine/Plugins/Experimental/PlainProps/Source/Tests/PlainPropsUeCoreTests.cpp
@@ -125,7 +125,7 @@ struct TScopedStructBinding : TScopedStructDeclaration<T, Occupancy, Runtime>
 	FStructSchemaId BindId;
 
 	TScopedStructBinding()
-	: BindId(IndexStructBindIdIfNeeded<Typename, Ids>(Super::DeclId))
+	: BindId(IndexStructBindIdIfNeeded<Ids, Typename>(Super::DeclId))
 	{
 		BindNativeStruct<CttiOf<T>, Runtime>(Runtime::GetSchemas(), BindId, Super::DeclId);
 	}
@@ -220,17 +220,17 @@ FBatchSaver::FBatchSaver(const FCustomBindings& CustomBase)
 template<class T>
 void FBatchSaver::Save(T&& Object) 
 {
-	FDualStructSchemaId Id = IndexStructDualId<std::remove_reference_t<T>, FIds>();
-	SavedObjects.Emplace(Id.Decl, SaveStruct(&Object, Id.Bind, {GTypes, GSchemas, Customs, Scratch}));
+	FDualStructSchemaId Id = IndexStructDualId<FIds, TTypename<std::remove_reference_t<T>>>();
+	SavedObjects.Emplace(Id.DeclId, SaveStruct(&Object, Id.BindId, {GTypes, GSchemas, Customs, Scratch}));
 }
 
 template<class T>
 bool FBatchSaver::SaveDelta(const T& Object, const T& Default) 
 {
-	FDualStructSchemaId Id = IndexStructDualId<std::remove_reference_t<T>, FIds>();
-	if (FBuiltStruct* Delta = SaveStructDelta(&Object, &Default, Id.Bind, {GTypes, GSchemas, Customs, Scratch}))
+	FDualStructSchemaId Id = IndexStructDualId<FIds, TTypename<std::remove_reference_t<T>>>();
+	if (FBuiltStruct* Delta = SaveStructDelta(&Object, &Default, Id.BindId, {GTypes, GSchemas, Customs, Scratch}))
 	{
-		SavedObjects.Emplace(Id.Decl, MoveTemp(Delta));
+		SavedObjects.Emplace(Id.DeclId, MoveTemp(Delta));
 		return true;
 	}
 	return false;
