PlainProps: Use the type correct array of constants when writing explicit enum constants for enum schemas of width B8, B16 and B32.

#rnx

diff --git a/Engine/Plugins/Experimental/PlainProps/Source/Private/PlainPropsWrite.cpp b/Engine/Plugins/Experimental/PlainProps/Source/Private/PlainPropsWrite.cpp
--- a/Engine/Plugins/Experimental/PlainProps/Source/Private/PlainPropsWrite.cpp
+++ b/Engine/Plugins/Experimental/PlainProps/Source/Private/PlainPropsWrite.cpp
@@ -416,7 +416,7 @@ static void WriteEnumConstantsAs(TArray64<uint8>& Out, TConstArrayView<uint64> C
 	{
 		Tmp.Add(IntCastChecked<IntType>(Constant));
 	}
-	WriteArray(Out, Constants);
+	WriteArray(Out, Tmp);
 }
 
 static void WriteEnumConstants(TArray64<uint8>& Out, ELeafWidth Width, TConstArrayView<uint64> Constants)
