PlainProps changes

* New custom binding load APIs FMemberLoader/FStructLoadView/FRangeLoadView
* New TMemberSerializer APIs for saving/loading templated custom bindings
* New TRangeSaver API for saving ranges in templated custom bindings
* New FCustomLoadIdsPlan to bind type-erased custom binding template parameters
* New TScopedStructBinding that works for custom and schema bound structs
* New FLeafRangeView::As<T>() that handles FBoolRangeView too
* New ReflectLeaf<T> helper
* New TOccupancyOf trait to handle custom bindings that are sometimes disabled
* New TPairBinding to handle non-default constructible pairs
* New TMapDeltaBinding
* New TSetDeltaBindig implementation
* New unit tests to cover some of the above
* Fix various bugs

#rnx

diff --git a/Engine/Plugins/Experimental/PlainProps/Source/Private/PlainPropsBind.cpp b/Engine/Plugins/Experimental/PlainProps/Source/Private/PlainPropsBind.cpp
--- a/Engine/Plugins/Experimental/PlainProps/Source/Private/PlainPropsBind.cpp
+++ b/Engine/Plugins/Experimental/PlainProps/Source/Private/PlainPropsBind.cpp
@@ -147,10 +147,21 @@ void* FLeafRangeAllocator::Allocate(FUnpackedLeafType Leaf, uint64 Num)
 
 ////////////////////////////////////////////////////////////////////////////////////////////////
 
-void FCustomBindings::BindStruct(FStructSchemaId BindId, FStructSchemaId DeclId, ICustomBinding& Binding)
+static bool operator==(FInnerStruct A, FInnerStruct B) { return A.Name == B.Name && A.BindId == B.BindId; }
+
+void FCustomBindings::BindStruct(FStructSchemaId BindId, FStructSchemaId DeclId, TConstArrayView<FInnerStruct> LoweredInners, ICustomBinding& Binding)
 {
 	checkf(!Find(BindId), TEXT("'%s' already bound"), *Debug.Print(BindId));
-	Entries.Emplace(BindId, DeclId, &Binding);
+	FInnersHandle Inners = StoreInners(LoweredInners);
+	check(Algo::Compare(LoweredInners, FetchInners(Inners)));
+	Entries.Emplace(BindId, DeclId, &Binding, Inners);
+}
+
+const ICustomBinding* FCustomBindings::FindStruct(FStructSchemaId BindId, TConstArrayView<FInnerStruct>& OutLoweredInners) const
+{
+	FCustomBindingEntry Entry = Find(BindId);
+	OutLoweredInners = FetchInners(Entry.LoweredInners);
+	return Entry.Binding;
 }
 
 void FCustomBindings::DropStruct(FStructSchemaId BindId)
@@ -186,6 +197,37 @@ FCustomBindingEntry	FCustomBindings::Find(FStructSchemaId BindId) const
 	return Base ? Base->Find(BindId) : FCustomBindingEntry{};
 }
 
+const FCustomBindings& FCustomBindings::GetBottom() const
+{
+	const FCustomBindings* It;
+	for (It = this; It->Base; It = It->Base) {}
+	return *It;
+}
+
+FInnersHandle FCustomBindings::StoreInners(TConstArrayView<FInnerStruct> In)
+{
+	if (uint32 Num = static_cast<uint32>(In.Num()))
+	{
+		// Optimizable: Inline encode if Num == 1
+		TArray<FInnerStruct>& All = GetBottom().AllLoweredInners;
+		FInnersHandle Out = { Num, static_cast<uint32>(All.Num()) };
+		All.Append(In);
+		return Out;
+	}
+
+	return {};
+}
+
+TConstArrayView<FInnerStruct> FCustomBindings::FetchInners(FInnersHandle In) const
+{
+	if (int32 Num = static_cast<int32>(In.Num))
+	{
+		const TArray<FInnerStruct>& All = GetBottom().AllLoweredInners;
+		return { &All[In.Idx], Num };
+	}
+	return {};
+}
+
 ////////////////////////////////////////////////////////////////////////////////////////////////
 
 FSchemaBindings::~FSchemaBindings()
diff --git a/Engine/Plugins/Experimental/PlainProps/Source/Private/PlainPropsInternalRead.h b/Engine/Plugins/Experimental/PlainProps/Source/Private/PlainPropsInternalRead.h
--- a/Engine/Plugins/Experimental/PlainProps/Source/Private/PlainPropsInternalRead.h
+++ b/Engine/Plugins/Experimental/PlainProps/Source/Private/PlainPropsInternalRead.h
@@ -39,6 +39,17 @@ inline uint64 GrabRangeNum(ERangeSizeType MaxSize, FByteReader& ByteIt, FBitCach
 	return 0;
 }
 
+inline FMemoryView GrabRangeValues(uint64 Num, FMemberType InnerType, /* in-out */ FByteReader& ByteIt)
+{
+	if (Num == 0)
+	{
+		return FMemoryView();
+	}
+
+	uint64 NumBytes = InnerType.IsLeaf() ? GetLeafRangeSize(Num, InnerType.AsLeaf()) : ByteIt.GrabVarIntU();
+	return ByteIt.GrabSlice(NumBytes);
+}
+
 //////////////////////////////////////////////////////////////////////////
 
 template<class T>
diff --git a/Engine/Plugins/Experimental/PlainProps/Source/Private/PlainPropsLoad.cpp b/Engine/Plugins/Experimental/PlainProps/Source/Private/PlainPropsLoad.cpp
--- a/Engine/Plugins/Experimental/PlainProps/Source/Private/PlainPropsLoad.cpp
+++ b/Engine/Plugins/Experimental/PlainProps/Source/Private/PlainPropsLoad.cpp
@@ -1,9 +1,11 @@
 // Copyright Epic Games, Inc. All Rights Reserved.
 
 #include "PlainPropsLoad.h"
+#include "PlainPropsLoadMember.h"
 #include "PlainPropsBind.h"
 #include "PlainPropsInternalFormat.h"
 #include "PlainPropsInternalRead.h"
+#include "Algo/Find.h"
 #include "Containers/Set.h"
 #include "Misc/Optional.h"
 #include <type_traits>
@@ -13,8 +15,16 @@ namespace PlainProps
 
 struct FLoadStructMemcpy
 {
-	uint32			Size;
-	uint32			Offset;
+	uint32					Size;
+	uint32					Offset;
+};
+
+// ICustomBinding with internal type-erased/lowered struct schemas
+struct FCustomLoadIdsPlan
+{
+	const ICustomBinding*	Binding;
+	uint32					NumLoadIds;
+	FStructSchemaId			LoadIds[0];
 };
 
 // Describes how to load a saved struct into the matching in-memory representation
@@ -35,6 +45,13 @@ public:
 		check(&Custom == &AsCustom());
 	}
 
+	explicit FLoadStructPlan(const FCustomLoadIdsPlan& Custom)
+	: Handle(uint64(&Custom) | CustomMask | LoadIdsBit)
+	{
+		check(Custom.Binding == &AsCustom());
+		check(GetInnerLoadIds() == Custom.LoadIds);
+	}
+
 	// @param OffsetWidth Usage unimplemented, store size and offsets as 8/16/32/64-bit 
 	explicit FLoadStructPlan(const FSchemaBinding& Schema, ELeafWidth OffsetWidth, bool bSparse)
 	: Handle(uint64(&Schema) | SchemaBit | (bSparse ? SparseBit : 0) | (uint64(OffsetWidth) << SchemaOffsetShift))
@@ -50,14 +67,17 @@ public:
 	bool						IsMemcpy() const		{ return (Handle & LoMask) == MemcpyMask; }
 	bool						IsCustom() const		{ return (Handle & LoMask) == CustomMask; }
 	FLoadStructMemcpy			AsMemcpy() const		{ check(IsMemcpy()); return { static_cast<uint32>(Handle >> 32), static_cast<uint32>(Handle) >> 1 }; }
-	const ICustomBinding&		AsCustom() const		{ check(IsCustom()); return *AsPtr<ICustomBinding>(); }
 	const FSchemaBinding&		AsSchema() const		{ check(IsSchema()); return *AsPtr<FSchemaBinding>(); }
+	const ICustomBinding&		AsCustom() const;
 
+	const FStructSchemaId*		GetInnerLoadIds()		{ return (Handle & TagMask) == LoadIdsMask ? AsPtr<FCustomLoadIdsPlan>()->LoadIds : nullptr; }
 	ELeafWidth					GetOffsetWidth() const	{ check(IsSchema()); return static_cast<ELeafWidth>((Handle & SchemaOffsetMask) >> SchemaOffsetShift); }
 
 private:
 	static constexpr uint64 SparseBit			= uint64(1) << FPlatformMemory::KernelAddressBit;
-	static constexpr uint64 PtrMask				= ~(SparseBit | 0b111);
+	static constexpr uint64 LoadIdsBit			= SparseBit;
+	static constexpr uint64 TagMask				= SparseBit | 0b111;
+	static constexpr uint64 PtrMask				= ~(TagMask);
 	static constexpr uint64 LoMask				= 0b11;
 	static constexpr uint64 MemcpyMask			= 0b00;
 	static constexpr uint64 CustomMask			= 0b10;
@@ -65,7 +85,8 @@ private:
 	static constexpr uint64 SparseSchemaMask	= SchemaBit | SparseBit;
 	static constexpr uint64 SchemaOffsetShift	= 1;
 	static constexpr uint64 SchemaOffsetMask	= 0b110;
-	
+	static constexpr uint64 LoadIdsMask			= CustomMask | LoadIdsBit;
+
 	template<typename T>
 	const T* AsPtr() const
 	{
@@ -76,6 +97,13 @@ private:
 	uint64						Handle = 0;
 };
 
+const ICustomBinding& FLoadStructPlan::AsCustom() const
+{ 
+	check(IsCustom());
+	return (Handle & LoadIdsBit) ? *AsPtr<FCustomLoadIdsPlan>()->Binding : *AsPtr<ICustomBinding>();
+}
+
+
 ////////////////////////////////////////////////////////////////////////////
 
 static uint16 CountEnums(const FStructSchema& Schema)
@@ -139,7 +167,7 @@ struct FLoadIdMapping
 	FStructSchemaId BindId; // Runtime bind id
 };
 
-// Helps load type-erased structs with ExplicitBindName by allocating new load-time struct ids
+// Helps load type-erased/lowered structs with ExplicitBindName by allocating new load-time struct ids
 class FLoadIdBinder
 {
 public:
@@ -151,7 +179,8 @@ public:
 	FStructSchemaId					BindLoadId(FStructSchemaId ReadId, FStructSchemaId BindId)
 	{
 		FStructSchemaId DeclId = DeclIds[ReadId.Idx];
-		return BindId == DeclId ? ReadId : MapLoadId(ReadId, BindId);
+		FStructSchemaId LoadId = BindId == DeclId ? ReadId : MapLoadId(ReadId, BindId);
+		return LoadId;
 	}
 	FLoadIdMapping					GetMapping(int32 Idx) const { return Mappings[Idx]; }
 	int32							NumMappings() const { return Mappings.Num(); }
@@ -344,11 +373,40 @@ struct FMissingBinding : ICustomBinding
 	FMissingBinding(const TCHAR* InType) : Type(InType) {}
 	const TCHAR* Type;
 	virtual void SaveCustom(FMemberBuilder& Dst, const void* Src, const void* Default, const FSaveContext& Ctx) override { checkf(false, TEXT("Can't save %s"), Type); }
-	virtual void LoadCustom(void* Dst, FStructView Src, ECustomLoadMethod Method, const FLoadBatch& Batch) const override { checkf(false, TEXT("Can't load %s"), Type); }
+	virtual void LoadCustom(void* Dst, FStructLoadView Src, ECustomLoadMethod Method) const override { checkf(false, TEXT("Can't load %s"), Type); }
 	virtual bool DiffCustom(const void* StructA, const void* StructB) const override { checkf(false, TEXT("Can't diff %s"), Type); return false; }
 };
-static FMissingBinding GMemberlessBinding(TEXT("Memberless struct binding"));
-static FMissingBinding GUnboundBinding{TEXT("Type-erased struct binding")};
+static FMissingBinding GMemberlessBinding(TEXT("memberless struct binding"));
+static FMissingBinding GUnboundBinding{TEXT("type-erased/lowered struct binding")};
+
+template<class T>
+static TOptional<TConstArrayView<T>> GetContiguousSubset(TConstArrayView<T> View, const TBitArray<>& Keep)
+{
+	check(Keep.Num() == View.Num());
+	const int32 Num = Keep.CountSetBits();
+	check(Num < Keep.Num());
+
+	if (Num == 0)
+	{
+		return TConstArrayView<T>{};
+	}
+
+	int32 StartIdx = Keep.Find(true);
+	check(StartIdx != INDEX_NONE);
+	if (StartIdx + 1 == Keep.Num())
+	{
+		return View.Slice(StartIdx, Num);
+	}
+
+	int32 EndIdx = Keep.FindFrom(false, StartIdx + 1);
+	check(EndIdx != INDEX_NONE);
+	if (EndIdx - StartIdx == Num)
+	{
+		return View.Slice(StartIdx, Num);
+	}
+
+	return NullOpt;
+}
 
 struct FLoadPlanner
 {
@@ -436,7 +494,7 @@ private:
 			check(OutReadId[-1].Idx < NumReadSchemas);
 		}
 
-		// Copy cloned subset schemas and patch up plan -> schema pointers
+		// Copy cloned subset schemas and patch up plan pointers
 		if (SubsetSchemaData.Num() > 0)
 		{
 			void* OutSubsetData = Align(OutReadId, alignof(FSchemaBinding));
@@ -448,10 +506,18 @@ private:
 			{
 				if (int32 Size = SubsetSchemaSizes[Idx])
 				{
-					check(IsAligned(Size, alignof(FSchemaBinding)));
-					check(Plans[Idx].IsSchema());
-					bool bSparse = Plans[Idx].IsSparseSchema();
-					Out->Plans[Idx] = FLoadStructPlan(*reinterpret_cast<const FSchemaBinding*>(It), ELeafWidth::B32, bSparse);
+					if (Plans[Idx].IsSchema())
+					{
+						check(IsAligned(Size, alignof(FSchemaBinding)));
+						bool bSparse = Plans[Idx].IsSparseSchema();
+						Out->Plans[Idx] = FLoadStructPlan(*reinterpret_cast<const FSchemaBinding*>(It), ELeafWidth::B32, bSparse);
+					}
+					else
+					{
+						check(Plans[Idx].IsCustom());
+						check(IsAligned(Size, alignof(FCustomLoadIdsPlan)));
+						Out->Plans[Idx] = FLoadStructPlan(*reinterpret_cast<const FCustomLoadIdsPlan*>(It));
+					}
 					It += Size;
 				}
 			}
@@ -469,14 +535,135 @@ private:
 		SubsetSchemaSizes[Mapping.LoadId.Idx] = static_cast<uint32>(SubsetSchemaData.Num() - SubsetSchemaOffset);
 	}
 
-	FLoadStructPlan CreatePlanInner(FLoadIdMapping Mapping, FLoadIdBinder& LoadIds)
+	// Schema iterator that scans forward to locate inner struct schema ids
+	class FInnerStructSchemaIterator
 	{
-		if (const ICustomBinding* Custom = Customs.FindStruct(Mapping.BindId))
+	public:
+		FInnerStructSchemaIterator(const FStructSchema& Schema) : FInnerStructSchemaIterator(Schema, UsesSuper(Schema.Inheritance)) {}
+
+		FStructSchemaId GrabMemberStruct(FMemberId InName)
 		{
-			return FLoadStructPlan(*Custom);
+			while (true)
+			{
+				FMemberId Name = Names[NameIdx];
+				FMemberType Type = NamedTypes[NameIdx];
+				FMemberType InnermostType = Type.IsRange() ? GrabInnerRangeTypes(RangeTypes, /* in-out */ RangeTypeIdx).Last() : Type;
+				++NameIdx;
+				InnerSchemaIt += IsStructOrEnum(InnermostType);
+				
+				if (Name == InName)
+				{
+					check(InnermostType.IsStruct());
+					return static_cast<FStructSchemaId>(InnerSchemaIt[-1]);
+				}
+			}
 		}
 
+	private:
+		FMemberType const* const			NamedTypes;
+		const TConstArrayView<FMemberId>	Names;
+		const TConstArrayView<FMemberType>	RangeTypes;
+		
+		int32								NameIdx = 0;
+		int32								RangeTypeIdx = 0;
+		const FSchemaId*					InnerSchemaIt;
+
+		FInnerStructSchemaIterator(const FStructSchema& Schema, bool bSkipSuper)
+		: NamedTypes(FStructSchema::GetMemberTypes(Schema.Footer) + bSkipSuper)
+		, Names(Schema.GetMemberNames())
+		, RangeTypes(Schema.GetRangeTypes())
+		, InnerSchemaIt(Schema.GetInnerSchemas() + /* super schema */ (bSkipSuper && !Schema.GetMemberTypes()[0].AsStruct().IsDynamic))
+		{}
+	};
+
+	[[nodiscard]] FCustomLoadIdsPlan* CreateCustomLoadIdsPlan(const ICustomBinding* Custom, const FStructSchema& Schema, TConstArrayView<FInnerStruct> Inners, /* in-out */ FLoadIdBinder& LoadIds)
+	{
+		const uint32 Num = static_cast<uint32>(Inners.Num());
+		check(Num);
+		check(Schema.NumInnerSchemas > 0);
+
+		// Calculate padding and size needed
+		uint32 Size = offsetof(FCustomLoadIdsPlan, LoadIds) + sizeof(FStructSchemaId) * Num;
+		uint32 Pad = Align(SubsetSchemaData.Num(), alignof(FCustomLoadIdsPlan)) - SubsetSchemaData.Num();
+		SubsetSchemaData.AddUninitialized(Pad + Size);
+
+		// Construct and init header
+		FCustomLoadIdsPlan Header = {Custom, Num};
+		FCustomLoadIdsPlan* Out = new (&SubsetSchemaData[SubsetSchemaData.Num() - Size]) FCustomLoadIdsPlan {Header};
+		check(IsAligned(Out, alignof(FCustomLoadIdsPlan)));
+		FStructSchemaId* OutLoadIdIt = Out->LoadIds;
+	
+		// Populate FLoadIdMapping and init Out->LoadIds
+		FInnerStructSchemaIterator SchemaIt(Schema);
+		for (FInnerStruct Inner : Inners)
+		{
+			// Todo: Rename FReadBatchId ReadId everywhere
+			FStructSchemaId ReadId_ = SchemaIt.GrabMemberStruct(Inner.Name);
+			*OutLoadIdIt++ = LoadIds.BindLoadId(ReadId_, Inner.BindId);
+		}
+		check((void*)OutLoadIdIt == SubsetSchemaData.GetData() + SubsetSchemaData.Num());
+
+		// Return unstable pointer that's later replaced in CreateBatch()
+		return Out;
+	}
+
+	// Get subset of Inners present in schema being read from
+	static TConstArrayView<FInnerStruct> GetLoweredMembers(TConstArrayView<FInnerStruct> Inners, TConstArrayView<FMemberId> Names, TArray<FInnerStruct, TInlineAllocator<8>>& TmpSubset)
+	{
+		bool bKeepAll = true;
+		TBitArray<> Keep(true, Inners.Num());
+		TConstArrayView<FMemberId> ScanNames = Names;
+		for (int32 Idx = 0; Idx < Inners.Num(); ++Idx)
+		{
+			if (const FMemberId* Name = Algo::Find(ScanNames, Inners[Idx].Name))
+			{
+				// Names must appear in order, limit future searches to later names
+				ScanNames.RightChopInline(1 + Name - &ScanNames[0]);
+			}
+			else
+			{
+				bKeepAll = false;
+				Keep[Idx] = false;
+			}
+		}
+
+		if (bKeepAll)
+		{
+			return Inners;
+		}
+		else if (TOptional<TConstArrayView<FInnerStruct>> ContiguousSubset = GetContiguousSubset(Inners, Keep))
+		{
+			return ContiguousSubset.GetValue();
+		}
+
+		for (int32 Idx = 0; Idx < Inners.Num(); ++Idx)
+		{
+			if (Keep[Idx])
+			{
+				TmpSubset.Add(Inners[Idx]);
+			}
+		}
+		return TmpSubset;
+	}
+
+	FLoadStructPlan CreatePlanInner(FLoadIdMapping Mapping, FLoadIdBinder& LoadIds)
+	{
 		const FStructSchema& From = ResolveStructSchema(ReadId, Mapping.ReadId);
+		
+		TConstArrayView<FInnerStruct> LoweredInners;
+		if (const ICustomBinding* Custom = Customs.FindStruct(Mapping.BindId, /* out */ LoweredInners))
+		{
+			TArray<FInnerStruct, TInlineAllocator<8>> TmpSubset;
+			TConstArrayView<FInnerStruct> LoweredMembers = GetLoweredMembers(LoweredInners, From.GetMemberNames(), /* out */ TmpSubset);
+
+			if (LoweredMembers.IsEmpty())
+			{
+				return FLoadStructPlan(*Custom);
+			}
+			
+			return FLoadStructPlan(*CreateCustomLoadIdsPlan(Custom, From, LoweredMembers, LoadIds));
+		}
+
 		if (From.NumMembers)
 		{
 			if (const FSchemaBinding* To = Schemas.FindStruct(Mapping.BindId))
@@ -546,7 +733,7 @@ inline static FMemberBindType ToBindType(FMemberType Member)
 class FRangeLoader
 {
 public:
-	static void LoadRangeView(uint8* Member, FRangeView Src, ERangeSizeType MaxSize, TConstArrayView<FRangeBinding> Bindings, const FLoadBatch& Batch)
+	static void LoadView(uint8* Member, FRangeLoadView Src, TConstArrayView<FRangeBinding> Bindings)
 	{
 		TArray<FMemberBindType, TFixedAllocator<16>> InnerTypes;
 		InnerTypes.Add(ToBindType(Src.Schema.ItemType));
@@ -558,14 +745,14 @@ public:
 			}
 		}
 		check(Bindings.Num() == InnerTypes.Num());
-		
-		FOptionalStructSchemaId StructSchema = InnerTypes.Last().IsStruct() ? static_cast<FOptionalStructSchemaId>(Src.Schema.InnermostSchema) : NoId;
-		
-		FLoadRangePlan Plan = { MaxSize, StructSchema, InnerTypes, Bindings.GetData() };
+			
+		FOptionalStructSchemaId InnermostStruct = InnerTypes.Last().IsStruct() ? static_cast<FOptionalStructSchemaId>(Src.Schema.InnermostId) : NoId;
+		ERangeSizeType Unused = ERangeSizeType::Uni;
+		FLoadRangePlan Plan = { Unused, InnermostStruct, InnerTypes, Bindings.GetData() };
 		
 		FByteReader ByteIt(Src.Values);
 		FBitCacheReader BitIt;
-		LoadRangePlan(Member, Src.NumItems, ByteIt, BitIt, Batch, Plan);
+		LoadRangePlan(Member, Src.NumItems, ByteIt, BitIt, Src.Schema.Batch, Plan);
 		ByteIt.CheckEmpty();
 	}
 
@@ -655,9 +842,20 @@ public:
 	static void CopyRangeValues(const FConstructedItems& Items, FByteReader& ByteIt, FBitCacheReader&, const FLoadBatch& Batch, FStructSchemaId Id)
 	{
 		uint64 ItemSize = Items.Size;
-		for (uint8* It = Items.Data, *End = It + Items.NumBytes(); It != End; It += ItemSize)
+		uint8* End = Items.Data + Items.NumBytes();
+		if (Items.bUnconstructed)
+		{
+			for (uint8* It = Items.Data; It != End; It += ItemSize)
+			{
+				PlainProps::ConstructAndLoadStruct(It, FByteReader(ByteIt.GrabSkippableSlice()), Id, Batch);
+			}
+		}
+		else
 		{
-			PlainProps::LoadStruct(It, FByteReader(ByteIt.GrabSkippableSlice()), Id, Batch);
+			for (uint8* It = Items.Data; It != End; It += ItemSize)
+			{
+				PlainProps::LoadStruct(It, FByteReader(ByteIt.GrabSkippableSlice()), Id, Batch);
+			}
 		}
 	}
 	
@@ -812,16 +1010,11 @@ void LoadStruct(void* Dst, FByteReader Src, FStructSchemaId LoadId, const FLoadB
 	}
 	else
 	{
-		FStructSchemaHandle ReadSchema{Batch.GetReadId(LoadId), Batch.ReadId};
-		Plan.AsCustom().LoadCustom(Dst, { ReadSchema, Src }, ECustomLoadMethod::Assign, Batch);
+		FSchemaLoadHandle Schema{LoadId, Batch};
+		Plan.AsCustom().LoadCustom(Dst, { Schema, Src }, ECustomLoadMethod::Assign);
 	}
 }
 
-void LoadStruct(void* Dst, FStructView Src, const FLoadBatch& Batch)
-{
-	LoadStruct(Dst, Src.Values, Src.Schema.Id, Batch);
-}
-
 void ConstructAndLoadStruct(void* Dst, FByteReader Src, FStructSchemaId Id, const FLoadBatch& Batch)
 {
 	FLoadStructPlan Plan = Batch[Id];
@@ -834,15 +1027,116 @@ void ConstructAndLoadStruct(void* Dst, FByteReader Src, FStructSchemaId Id, cons
 	}
 	else
 	{
-		FStructSchemaHandle ReadSchema{Id, Batch.ReadId};
-		Plan.AsCustom().LoadCustom(Dst, { ReadSchema, Src }, ECustomLoadMethod::Construct, Batch);
+		FSchemaLoadHandle Schema{Id, Batch};
+		Plan.AsCustom().LoadCustom(Dst, { Schema, Src }, ECustomLoadMethod::Construct);
+	}
+}
+
+////////////////////////////////////////////////////////////////////////////
+
+FRangeLoadView FNestedRangeLoadIterator::operator*() const 
+{
+	FByteReader PeekBytes = ByteIt;
+	FBitCacheReader PeekBits = BitIt;
+
+	FRangeLoadSchema OutSchema = { Schema.NestedItemTypes[0], Schema.InnermostId, /* Only valid for nested ranges */ Schema.NestedItemTypes + 1, Schema.Batch };
+	uint64 OutNumItems = GrabRangeNum(Schema.ItemType.AsRange().MaxSize, /* in-out */ PeekBytes,  /* in-out */ PeekBits);
+	FMemoryView OutValues = GrabRangeValues(OutNumItems, OutSchema.ItemType, /* in-out */ PeekBytes);
+
+	return { OutSchema, OutNumItems, OutValues };
+}
+
+void FNestedRangeLoadIterator::operator++()
+{
+	uint64 Num = GrabRangeNum(Schema.ItemType.AsRange().MaxSize, /* in-out */ ByteIt, /* in-out */ BitIt);
+	(void)GrabRangeValues(Num, Schema.NestedItemTypes[0], /* in-out */ ByteIt);
+}
+
+///////////////////////////////////////////////////////////////////////////
+
+FLeafRangeView FRangeLoadView::AsLeaves() const
+{
+	FUnpackedLeafType Leaf = Schema.ItemType.AsLeaf();
+	FOptionalEnumSchemaId EnumId = static_cast<FOptionalEnumSchemaId>(Schema.InnermostId);
+	const uint8* Data = static_cast<const uint8*>(Values.GetData());
+	return { Leaf, Schema.Batch.ReadId, EnumId, NumItems, Data };
+}
+
+FStructRangeLoadView FRangeLoadView::AsStructs() const
+{
+	check(IsStructRange());
+	FStructSchemaId LoadId = static_cast<FStructSchemaId>(Schema.InnermostId.Get());
+	return { NumItems, Values, {LoadId, Schema.Batch} };
+}
+
+FNestedRangeLoadView FRangeLoadView::AsRanges() const
+{
+	check(IsNestedRange());
+	return { NumItems, Values, Schema };
+}
+
+////////////////////////////////////////////////////////////////////////////
+
+static FStructView ToReadView(FStructLoadView In)
+{
+	FStructSchemaHandle ReadSchema = { In.Schema.Batch.GetReadId(In.Schema.LoadId), In.Schema.Batch.ReadId };
+	return { ReadSchema, In.Values };
+}
+
+FMemberLoader::FMemberLoader(FStructLoadView In)
+: Reader(ToReadView(In))
+, LoadIdIt(In.Schema.Batch[In.Schema.LoadId].GetInnerLoadIds())
+, Batch(In.Schema.Batch)
+{}
+
+static FMemberType GetInnermostType(FRangeSchema Schema)
+{
+	FMemberType Inner = Schema.ItemType;
+	for (const FMemberType* It = Schema.NestedItemTypes; Inner.IsRange(); Inner = *It++) {}
+	return Inner;
+}
+
+FRangeLoadView FMemberLoader::GrabRange()
+{
+	const FRangeView In = Reader.GrabRange();
+
+	// Replace ReadId with LoadId
+	FOptionalSchemaId InnermostId = In.Schema.InnermostSchema;
+	if (LoadIdIt && InnermostId && GetInnermostType(In.Schema).IsStruct())
+	{
+		check(LoadIdIt->Idx > InnermostId.Get().Idx);
+		check(InnermostId == Batch.GetReadId(*LoadIdIt));
+		InnermostId = *LoadIdIt++;
 	}
+
+	FRangeLoadSchema OutSchema = { In.Schema.ItemType, InnermostId, In.Schema.NestedItemTypes, Batch };
+	return { OutSchema, In.NumItems, In.Values };
+}
+
+FStructLoadView FMemberLoader::GrabStruct()
+{
+	const FStructView In = Reader.GrabStruct();
+	FStructSchemaId LoadId = LoadIdIt ? *LoadIdIt++ : In.Schema.Id;
+	check(In.Schema.Id == Batch.GetReadId(LoadId));
+
+	return { {LoadId, Batch}, In.Values };
+}
+
+////////////////////////////////////////////////////////////////////////////
+
+void LoadRange(void* Dst, FRangeLoadView Src, TConstArrayView<FRangeBinding> Bindings)
+{
+	FRangeLoader::LoadView(static_cast<uint8*>(Dst), Src, Bindings);
 }
 
-void LoadRange(void* Dst, FRangeView Src, ERangeSizeType MaxSize, TConstArrayView<FRangeBinding> Bindings, const FLoadBatch& Batch)
+void LoadStruct(void* Dst, FStructLoadView Src)
 {
-	FRangeLoader::LoadRangeView(static_cast<uint8*>(Dst), Src, MaxSize, Bindings, Batch);
+	LoadStruct(Dst, Src.Values, Src.Schema.LoadId, Src.Schema.Batch);
 }
 
+void ConstructAndLoadStruct(void* Dst, FStructLoadView Src)
+{
+	ConstructAndLoadStruct(Dst, Src.Values, Src.Schema.LoadId, Src.Schema.Batch);
+}
 
 } // namespace PlainProps
\ No newline at end of file
diff --git a/Engine/Plugins/Experimental/PlainProps/Source/Private/PlainPropsRead.cpp b/Engine/Plugins/Experimental/PlainProps/Source/Private/PlainPropsRead.cpp
--- a/Engine/Plugins/Experimental/PlainProps/Source/Private/PlainPropsRead.cpp
+++ b/Engine/Plugins/Experimental/PlainProps/Source/Private/PlainPropsRead.cpp
@@ -145,40 +145,22 @@ FParametricTypeView ResolveUntranslatedParametricType(FReadBatchId Batch, FParam
 FLeafRangeView FRangeView::AsLeaves() const
 {
 	FUnpackedLeafType Leaf = Schema.ItemType.AsLeaf();
-
-	FLeafRangeView Out;
-	Out.Type = Leaf.Type;
-	Out.Width = Leaf.Width;
-	Out.Batch = Schema.Batch;
-	Out.Enum = static_cast<FOptionalEnumSchemaId>(Schema.InnermostSchema);
-	Out.NumItems = NumItems;
-	Out.Values = static_cast<const uint8*>(Values.GetData());
-
-	return Out;
+	FOptionalEnumSchemaId EnumId = static_cast<FOptionalEnumSchemaId>(Schema.InnermostSchema);
+	const uint8* Data = static_cast<const uint8*>(Values.GetData());
+	return {Leaf, Schema.Batch, EnumId, NumItems, Data };
 }
 
 FStructRangeView FRangeView::AsStructs() const
 {
 	check(IsStructRange());
-
-	FStructRangeView Out;
-	Out.NumItems = NumItems;
-	Out.Data = Values;
-	Out.Schema.Id = static_cast<FStructSchemaId>(Schema.InnermostSchema.Get());
-	Out.Schema.Batch = Schema.Batch;
-
-	return Out;
+	FStructSchemaId Id = static_cast<FStructSchemaId>(Schema.InnermostSchema.Get());
+	return {NumItems, Values, {Id, Schema.Batch}};
 }
 
 FNestedRangeView FRangeView::AsRanges() const
 {
 	check(IsNestedRange());
-
-	FNestedRangeView Out;
-	Out.NumItems = NumItems;
-	Out.Data = Values;
-	Out.Schema = Schema;
-	return Out;
+	return {NumItems, Values, Schema};
 }
 
 //////////////////////////////////////////////////////////////////////////
@@ -190,18 +172,6 @@ const FStructSchema& FStructSchemaHandle::ResolveSuper() const
 
 //////////////////////////////////////////////////////////////////////////
 
-static FMemoryView GrabRangeValues(uint64 Num, FMemberType InnerType, /* in-out */ FByteReader& ByteIt)
-{
-	if (Num == 0)
-	{
-		return FMemoryView();
-	}
-
-	uint64 NumBytes = InnerType.GetKind() == EMemberKind::Leaf ? GetLeafRangeSize(Num, InnerType.AsLeaf()) : ByteIt.GrabVarIntU();
-
-	return ByteIt.GrabSlice(NumBytes);
-}
-
 FRangeView FNestedRangeIterator::operator*() const 
 {
 	FByteReader PeekBytes = ByteIt;
diff --git a/Engine/Plugins/Experimental/PlainProps/Source/Private/PlainPropsSave.cpp b/Engine/Plugins/Experimental/PlainProps/Source/Private/PlainPropsSave.cpp
--- a/Engine/Plugins/Experimental/PlainProps/Source/Private/PlainPropsSave.cpp
+++ b/Engine/Plugins/Experimental/PlainProps/Source/Private/PlainPropsSave.cpp
@@ -4,6 +4,7 @@
 #include "PlainPropsBind.h"
 #include "PlainPropsInternalBuild.h"
 #include "PlainPropsInternalFormat.h"
+#include "PlainPropsSaveMember.h"
 #include "Algo/Find.h"
 #include <type_traits>
 
@@ -194,9 +195,9 @@ template<class SaverType, typename InnerContextType>
 	return { Member.InnerTypes + 1, Member.RangeBindings + 1, static_cast<uint16>(Member.NumRanges - 1), Member.InnermostSchema };
 }
 
-[[nodiscard]] static FBuiltRange* SaveLeafRangeBinding(FScratchAllocator& Scratch, const void* Range, const ILeafRangeBinding& Binding, FUnpackedLeafType Leaf)
+FBuiltRange* SaveLeafRange(const void* Range, const ILeafRangeBinding& Binding, FUnpackedLeafType Leaf, const FSaveContext& Ctx)
 {
-	FLeafRangeAllocator Allocator(Scratch, Leaf);
+	FLeafRangeAllocator Allocator(Ctx.Scratch, Leaf);
 	Binding.SaveLeaves(Range, Allocator);
 	return Allocator.GetAllocatedRange();
 }
@@ -226,7 +227,7 @@ ELeafWidth GetArithmeticWidth(FLeafBindType Leaf)
 
 	if (Binding.IsLeafBinding())
 	{
-		return SaveLeafRangeBinding(Ctx.Scratch, Range, Binding.AsLeafBinding(), UnpackNonBitfield(InnerType.AsLeaf()));
+		return SaveLeafRange(Range, Binding.AsLeafBinding(), UnpackNonBitfield(InnerType.AsLeaf()), Ctx);
 	}
 
 	const IItemRangeBinding& ItemBinding = Binding.AsItemBinding();
@@ -564,4 +565,14 @@ FBuiltStruct* SaveStructDelta(const void* Struct, const void* Default, FStructSc
 	return Out.IsEmpty() ? nullptr : Out.BuildAndReset(Ctx.Scratch, *Declaration, Ctx.Declarations.GetDebug());
 }
 
+////////////////////////////////////////////////////////////////////////////////////////////////
+
+FRangeSaverBase::FRangeSaverBase(FScratchAllocator& Scratch, uint64 Num, SIZE_T ItemSize)
+: Range(FBuiltRange::Create(Scratch, Num, ItemSize))
+, It(Range->Data)
+#if DO_CHECK
+, End(It + Num * ItemSize)
+#endif
+{}
+
 } // namespace PlainProps
\ No newline at end of file
diff --git a/Engine/Plugins/Experimental/PlainProps/Source/Private/PlainPropsUeCoreBindings.cpp b/Engine/Plugins/Experimental/PlainProps/Source/Private/PlainPropsUeCoreBindings.cpp
--- a/Engine/Plugins/Experimental/PlainProps/Source/Private/PlainPropsUeCoreBindings.cpp
+++ b/Engine/Plugins/Experimental/PlainProps/Source/Private/PlainPropsUeCoreBindings.cpp
@@ -46,20 +46,20 @@ void FTransformBinding::Save(FMemberBuilder& Dst, const FTransform& Src, const F
 }
 
 template<typename T>
-T GrabAndMemcpy(FMemberReader& Members)
+T GrabAndMemcpy(FMemberLoader& Members)
 {
 	T Out;
-	FStructView Struct = Members.GrabStruct();
+	FStructLoadView Struct = Members.GrabStruct();
 	Struct.Values.CheckSize(sizeof(T));
 	FMemory::Memcpy(&Out, Struct.Values.Peek(), sizeof(T));
 	return Out;
 }
 
-void FTransformBinding::Load(FTransform& Dst, FStructView Src, ECustomLoadMethod Method, const FLoadBatch& Batch) const
+void FTransformBinding::Load(FTransform& Dst, FStructLoadView Src, ECustomLoadMethod Method) const
 {
 	static_assert(std::is_same_v<decltype(FTransform().GetTranslation().X), double>);
 
-	FMemberReader Members(Src);
+	FMemberLoader Members(Src);
 
 	if (Method == ECustomLoadMethod::Construct)
 	{
diff --git a/Engine/Plugins/Experimental/PlainProps/Source/Public/PlainPropsBind.h b/Engine/Plugins/Experimental/PlainProps/Source/Public/PlainPropsBind.h
--- a/Engine/Plugins/Experimental/PlainProps/Source/Public/PlainPropsBind.h
+++ b/Engine/Plugins/Experimental/PlainProps/Source/Public/PlainPropsBind.h
@@ -21,13 +21,15 @@ struct FBuiltStruct;
 class FIdIndexerBase;
 struct FLoadBatch;
 class FMemberBuilder;
+class FRangeLoader;
 struct FSchemaBatch;
 class FScratchAllocator;
 class FStructBinding;
+struct FStructLoadView;
 class FRangeBinding;
 struct FSaveContext;
 struct FTypedRange;
-class IItemRangeBinding;
+struct IItemRangeBinding;
 template<class T> class TIdIndexer;
 
 ////////////////////////////////////////////////////////////////////////////////////////////////
@@ -245,10 +247,18 @@ struct FDualStructSchemaId
 {
 	FStructSchemaId BindId;
 	FStructSchemaId DeclId;
+
+	bool IsLowered() const { return BindId != DeclId; }
 };
 
 ////////////////////////////////////////////////////////////////////////////////////////////////
 
+struct FInnerStruct
+{
+	FMemberId				Name; // Of outer Range or Struct
+	FStructSchemaId			BindId;
+};
+
 enum ECustomLoadMethod { Construct, Assign };
 
 /// Load/save a struct with custom code to handle:
@@ -261,9 +271,15 @@ enum ECustomLoadMethod { Construct, Assign };
 struct ICustomBinding
 {
 	virtual ~ICustomBinding() {}
-	virtual void				SaveCustom(FMemberBuilder& Dst, const void* Src, const void* Default, const FSaveContext& Ctx) = 0;
-	virtual void				LoadCustom(void* Dst, FStructView Src, ECustomLoadMethod Method, const FLoadBatch& Batch) const = 0;
-	virtual bool				DiffCustom(const void* StructA, const void* StructB) const = 0;
+	virtual void							SaveCustom(FMemberBuilder& Dst, const void* Src, const void* Default, const FSaveContext& Ctx) = 0;
+	virtual void							LoadCustom(void* Dst, FStructLoadView Src, ECustomLoadMethod Method) const = 0;
+	virtual bool							DiffCustom(const void* StructA, const void* StructB) const = 0;
+};
+
+struct FInnersHandle
+{
+	uint32			Num = 0;
+	uint32			Idx = 0;
 };
 
 struct FCustomBindingEntry
@@ -271,6 +287,7 @@ struct FCustomBindingEntry
 	FStructSchemaId BindId;
 	FStructSchemaId DeclId;
 	ICustomBinding* Binding = nullptr;
+	FInnersHandle	LoweredInners;
 
 	explicit operator bool() const { return !!Binding; }
 };
@@ -282,35 +299,54 @@ public:
 	FCustomBindings(const FDebugIds& Dbg, const FCustomBindings* InBase = nullptr) : Base(InBase), Debug(Dbg) {}
 
 	// @param Binding must outlive this or call DropStruct()
-	PLAINPROPS_API void						BindStruct(FStructSchemaId BindId, FStructSchemaId DeclId, ICustomBinding& Binding);
+	PLAINPROPS_API void						BindStruct(FStructSchemaId BindId, FStructSchemaId DeclId, TConstArrayView<FInnerStruct> LoweredInners, ICustomBinding& Binding);
 	const ICustomBinding*					FindStruct(FStructSchemaId BindId) const		{ return Find(BindId).Binding; }
+	PLAINPROPS_API const ICustomBinding*	FindStruct(FStructSchemaId BindId, TConstArrayView<FInnerStruct>& OutLoweredInners) const;
 	FOptionalStructSchemaId					FindStructDeclId(FStructSchemaId BindId) const;
 	FCustomBindingEntry						FindStructToSave(FStructSchemaId BindId)		{ return Find(BindId); }
 	PLAINPROPS_API void						DropStruct(FStructSchemaId BindId);
 
 private:
 	PLAINPROPS_API FCustomBindingEntry		Find(FStructSchemaId BindId) const;
+	const FCustomBindings&					GetBottom() const;
+	FInnersHandle							StoreInners(TConstArrayView<FInnerStruct> In);
+	TConstArrayView<FInnerStruct>			FetchInners(FInnersHandle In) const;
 	
 	const FCustomBindings*					Base = nullptr;
 	TArray<FCustomBindingEntry, TInlineAllocator<8>>		Entries;
+	mutable TArray<FInnerStruct>			AllLoweredInners;
 	const FDebugIds&						Debug;
 };
 
 template<typename T>
-struct TCustomBind{ using Type = void; };
+struct TCustomBind { using Type = void; };
 
 template<typename T>
 using CustomBind = typename TCustomBind<T>::Type;
 
 ////////////////////////////////////////////////////////////////////////////////////////////////
 
+template<typename T>
+struct TOccupancyOf
+{
+	static constexpr EMemberPresence Value = EMemberPresence::AllowSparse;
+};
+
+struct FRequireAll
+{
+	static constexpr EMemberPresence Value = EMemberPresence::RequireAll;
+};
+
+////////////////////////////////////////////////////////////////////////////////////////////////
+
 class FConstructionRequest
 {
+	friend FRangeLoader;
+
 	void* const Range = nullptr;
 	const uint64 Num = 0;
 	uint64 Index = 0;
 
-	friend FRangeLoader;
 	FConstructionRequest(void* InRange, uint64 InNum) : Range(InRange), Num(InNum) {}
 	
 public:
@@ -417,12 +453,6 @@ struct FExistingItems
 	}
 };
 
-
-// TODO: Consider API changes so that
-//  - Leaf ranges can allocate data and fill it in directly
-//  - Leaf ranges can allocate more capacity than is used
-//  - Maybe separate Finalize function
-//  - Continue hiding intermediate details (e.g. FBuiltRange) and allocator
 struct FSaveRangeContext
 {
 	FGetItemsRequest		Request;		// Request to get items to be saved
@@ -430,9 +460,8 @@ struct FSaveRangeContext
 	uint64					Scratch[8]; 	// Scratch memory for IRangeBinding
 };
 
-class alignas(16) IItemRangeBinding
+struct alignas(16) IItemRangeBinding
 {
-public:
 	virtual void ReadItems(FSaveRangeContext& Ctx) const = 0;
 	virtual void MakeItems(FLoadRangeContext& Ctx) const = 0;
 };
@@ -476,33 +505,24 @@ public:
 	
 	// The returned ranges hide the internal representations so we can change format in the future, 
 	// e.g. store zeroes or 1.0f in some compact fashion or even var int encodings
-	template<Arithmetic T>
+	template<typename T>
 	auto As() const
 	{
-		check(Leaf == ReflectArithmetic<T>);
-		if constexpr (std::is_same_v<bool, T>)
-		{
-			return FBoolRangeView(static_cast<const uint8*>(Data), Num);
-		}
-		else
-		{
-			return TRangeView<T>(static_cast<const T*>(Data), Num);
-		}
-	}
-	
-	template<Enumeration T>
-	TRangeView<T> As() const
-	{
-		check(Leaf == ReflectEnum<T>);
+		check(Leaf == ReflectLeaf<T>);
 		return TRangeView<T>(static_cast<const T*>(Data), Num);
 	}
 
+	template<>
+	auto As<bool>() const
+	{
+		check(Leaf.Type == ELeafType::Bool && Leaf.Width == ELeafWidth::B8);
+		return FBoolRangeView(static_cast<const uint8*>(Data), Num);
+	}
 };
 
 // Specialized binding for transcoding leaf ranges
-class alignas(16) ILeafRangeBinding
+struct alignas(16) ILeafRangeBinding
 {
-public:
 	virtual void	SaveLeaves(const void* Range, FLeafRangeAllocator& Out) const = 0;
 	virtual void	LoadLeaves(void* Range, FLeafRangeLoadView Leaves) const = 0;
 	virtual int64	DiffLeaves(const void* RangeA, const void* RangeB) const = 0;
@@ -673,6 +693,14 @@ FStructSchemaId GetStructBindId()
 	return Id;
 }
 
+// Cached by function static
+template<class Ids, typename Struct>
+FDualStructSchemaId GetStructDualId()
+{
+	static FDualStructSchemaId Id = IndexStructDualId<Ids, TTypename<Struct>>();
+	return Id;
+}
+
 template<class Ids, typename Ctti>
 FTypeId IndexCttiName()
 {
@@ -748,63 +776,119 @@ FTypeId IndexParametricType(FTypeId TemplatedType, const std::tuple<Ts...>*)
 
 ////////////////////////////////////////////////////////////////////////////////////////////////
 
+// Works around C++'s lack of templated nullary constructors
 template<class Ids>
-struct TCustomInit {};
+struct TInit {};
 
+////////////////////////////////////////////////////////////////////////////////////////////////
 
-template<typename CustomBinding, class Runtime>
-FStructSchemaId BindCustomStructOnce()
+using InnerStructArray = TArray<FInnerStruct, TInlineAllocator<8>>;
+
+struct FCustomInit
 {
-	struct FBinding : FDualStructSchemaId, CustomBinding
-	{
-		using Type = typename CustomBinding::Type;
-		using Typename = TCustomTypename<CustomBinding>;
-		using Ids = typename Runtime::Ids;
+	InnerStructArray& OutLowered;
 
-		explicit FBinding(FTypeId DeclType = IndexStructName<Ids, ETypename::Decl, Typename>())
-		: FDualStructSchemaId(IndexStructDualId<Ids, Typename>(DeclType))
-		, CustomBinding(TCustomInit<Ids>{})
+	// Only needed for inner structs that might get type-erased/lowered
+	void RegisterInnerStruct(FDualStructSchemaId Id, TConstArrayView<FMemberId> Names)
+	{
+		if (Id.IsLowered())
 		{
-			Runtime::GetTypes().DeclareStruct(DeclId, DeclType, CustomBinding::MemberIds, CustomBinding::Occupancy);
-			Runtime::GetCustoms().BindStruct(BindId, DeclId, *this);
+			for (FMemberId Name : Names)
+			{
+				OutLowered.Add({Name, Id.BindId});
+			}
 		}
+	}
+};
 
-		~FBinding()
-		{
-			Runtime::GetCustoms().DropStruct(BindId);
-			Runtime::GetTypes().DropStructRef(DeclId);
-		}
+// Helps custom binding constructors create ids and register types that might get type-erased 
+template<class Ids>
+struct TCustomInit : FCustomInit, TInit<Ids>
+{};
 
-		virtual void SaveCustom(FMemberBuilder& Dst, const void* Src, const void* Default, const FSaveContext& Ctx) override
-		{
-			CustomBinding::Save(Dst, *static_cast<const Type*>(Src), static_cast<const Type*>(Default), Ctx);
-		}
+// Workaround for clang requiring FStructLoadView definition, make it dependent on any other type
+template<typename>
+struct TDependent { using StructLoadView = FStructLoadView; };
 
-		virtual void LoadCustom(void* Dst, FStructView Src, ECustomLoadMethod Method, const FLoadBatch& Batch) const override
-		{
-			CustomBinding::Load(*static_cast<Type*>(Dst), Src, Method, Batch);
-		}
+// Scoped custom binding for native type
+template<class T, class Runtime, typename CustomBinding = CustomBind<T>>
+struct TScopedStructBinding : FDualStructSchemaId, CustomBinding
+{
+	using Ids = typename Runtime::Ids;
+	using Typename = TCustomTypename<CustomBinding>;
 
-		virtual bool DiffCustom(const void* A, const void* B) const override
-		{
-			return CustomBinding::Diff(*static_cast<const Type*>(A), *static_cast<const Type*>(B));
-		}
-	};
+	explicit TScopedStructBinding(FTypeId DeclName = IndexStructName<Ids, ETypename::Decl, Typename>(), InnerStructArray Lowered = {})
+	: FDualStructSchemaId(IndexStructDualId<Ids, Typename>(DeclName))
+	, CustomBinding(TCustomInit<Ids>{{Lowered}})
+	{
+		Runtime::GetTypes().DeclareStruct(DeclId, DeclName, CustomBinding::MemberIds, TOccupancyOf<T>::Value);
+		Runtime::GetCustoms().BindStruct(BindId, DeclId, Lowered, *this);
+	}
 
-	static FBinding Binding;
-	return Binding.BindId;
+	~TScopedStructBinding()
+	{
+		Runtime::GetCustoms().DropStruct(BindId);
+		Runtime::GetTypes().DropStructRef(DeclId);
+	}
+
+	virtual void SaveCustom(FMemberBuilder& Dst, const void* Src, const void* Default, const FSaveContext& Ctx) override
+	{
+		CustomBinding::Save(Dst, *static_cast<const T*>(Src), static_cast<const T*>(Default), Ctx);
+	}
+
+	virtual void LoadCustom(void* Dst, TDependent<Ids>::StructLoadView Src, ECustomLoadMethod Method) const override 
+	{
+		CustomBinding::Load(*static_cast<T*>(Dst), Src, Method);
+	}
+
+	virtual bool DiffCustom(const void* A, const void* B) const override
+	{
+		return CustomBinding::Diff(*static_cast<const T*>(A), *static_cast<const T*>(B));
+	}
+};
+
+// Scoped schema binding for native type
+template<class T, class Runtime>
+struct TScopedStructBinding<T, Runtime, void>
+{
+	using Ids = typename Runtime::Ids;
+	using Typename = TTypename<T>;
+	
+	const FStructSchemaId DeclId;
+	const FStructSchemaId BindId;
+
+	TScopedStructBinding(FTypeId DeclName = IndexStructName<Ids, ETypename::Decl, Typename>())
+	: DeclId(DeclareNativeStruct<CttiOf<T>, Ids>(Runtime::GetTypes(), TOccupancyOf<T>::Value))
+	, BindId(IndexStructBindIdIfNeeded<Ids, Typename>(DeclId))
+	{
+		BindNativeStruct<CttiOf<T>, Runtime>(Runtime::GetSchemas(), BindId, DeclId);
+	}
+
+	~TScopedStructBinding()
+	{
+		Runtime::GetSchemas().DropStruct(BindId);
+		Runtime::GetTypes().DropStructRef(DeclId);
+	}
+};
+
+template<class Struct, typename CustomBinding, class Runtime>
+FStructSchemaId BindCustomStructOnce()
+{
+	static TScopedStructBinding<Struct, Runtime, CustomBinding> Instance;
+	return Instance.BindId;
 }
 
 template<class Struct, class CustomBinding, class Runtime>
 FMemberBindType BindMemberStruct(FOptionalSchemaId& OutSchema)
 {
+	// Todo: Consider auto-binding all nested structs
 	if constexpr (std::is_void_v<CustomBinding>)
 	{
 		OutSchema = FOptionalSchemaId(GetStructBindId<typename Runtime::Ids, Struct>());
 	}
 	else
 	{
-		OutSchema = FOptionalSchemaId(BindCustomStructOnce<CustomBinding, Runtime>());
+		OutSchema = FOptionalSchemaId(BindCustomStructOnce<Struct, CustomBinding, Runtime>());
 	}
 
 	return FMemberBindType(FStructType{EMemberKind::Struct, /* IsDynamic */ 0, /* IsSuper */ 0});
@@ -834,26 +918,25 @@ FMemberBindType BindInnermostType(FOptionalSchemaId& OutSchema)
 template<typename RangeBinding>
 constexpr uint32 CountRangeBindings()
 {
-	if constexpr (std::is_void_v<RangeBinding>)
-	{
-		return 0;
-	}
-	else
-	{
-		using InnerBinding = RangeBind<typename RangeBinding::ItemType>;
-		return 1 + CountRangeBindings<InnerBinding>();
-	}
+	using InnerBinding = RangeBind<typename RangeBinding::ItemType>;
+	return 1 + CountRangeBindings<InnerBinding>();
+}
+
+template<>
+constexpr uint32 CountRangeBindings<void>()
+{
+	return 0;
 }
 
 template<typename RangeBinding, uint32 NestLevel>
-struct TInnerType
+struct TInnermostType
 {
 	using InnerType = typename RangeBinding::ItemType;
-	using Type = TInnerType<RangeBind<InnerType>, NestLevel - 1>::Type;
+	using Type = TInnermostType<RangeBind<InnerType>, NestLevel - 1>::Type;
 };
 
 template<typename RangeBinding>
-struct TInnerType<RangeBinding, 1>
+struct TInnermostType<RangeBinding, 1>
 {
 	using Type = typename RangeBinding::ItemType;
 };
@@ -922,7 +1005,7 @@ FMemberBinding BindMember(uint64 Offset)
 		if constexpr (!std::is_void_v<RangeBinding>)
 		{
 			constexpr uint32 NumRangeBindings = CountRangeBindings<RangeBinding>();
-			using InnermostType = typename TInnerType<RangeBinding, NumRangeBindings>::Type;
+			using InnermostType = typename TInnermostType<RangeBinding, NumRangeBindings>::Type;
 
 			Out.RangeBindings = GetRangeBindings<RangeBinding, NumRangeBindings>();
 			Out.InnermostType = BindInnermostType<InnermostType, Ids>(Out.InnermostSchema);
@@ -1048,67 +1131,6 @@ struct FMemberBinder
 
 //////////////////////////////////////////////////////////////////////////
 
-template<typename T>
-inline constexpr FMemberType ReflectInnermostType()
-{
-	// DefaultStructType
-	return FMemberType(FStructType{EMemberKind::Struct, /* IsDynamic */ 0, /* IsSuper */ 0});;
-}
-
-template<Arithmetic T>
-inline constexpr FMemberType ReflectInnermostType()
-{
-	return ReflectArithmetic<T>.Pack();
-}
-template<Enumeration T>
-inline constexpr FMemberType ReflectInnermostType()
-{
-	return ReflectEnum<T>.Pack();
-}
-
-// TRangeMemberHelper init helper
-union FUninitializedMemberBindType
-{
-	FUninitializedMemberBindType() : Unused() {}
-	uint8_t			Unused;
-	FMemberBindType Value;
-};
-
-// Helps templated custom bindings save containers as ranges
-template<class RangeBinding>
-struct TRangeMemberHelper
-{
-	static constexpr uint16 NumRanges = CountRangeBindings<RangeBinding>();
-	static constexpr ERangeSizeType MaxSize = RangeSizeOf(typename RangeBinding::SizeType{});
-	using InnermostType = typename TInnerType<RangeBinding, NumRanges>::Type;
-
-	const FRangeBinding*			RangeBindings = nullptr;
-	FOptionalSchemaId				InnermostSchema;
-	FUninitializedMemberBindType	InnerBindTypes[NumRanges];
-	FMemberType						InnerSchemaTypes[NumRanges];
-
-	template<class Ids>
-	void Init()
-	{
-		RangeBindings = GetRangeBindings<RangeBinding, NumRanges>().GetData();
-		for (uint16 Idx = 0; Idx < NumRanges - 1; ++Idx)
-		{
-			ERangeSizeType Type = RangeBindings[Idx + 1].GetSizeType();
-			InnerBindTypes[Idx].Value = FMemberBindType(Type);
-			InnerSchemaTypes[Idx] = FMemberType(Type);
-		}
-		InnerBindTypes[NumRanges - 1].Value = BindInnermostType<InnermostType, Ids>(/* out */ InnermostSchema);
-		InnerSchemaTypes[NumRanges - 1] = ReflectInnermostType<InnermostType>();
-	}
-
-	FRangeMemberBinding MakeBinding(uint32 Offset) const
-	{
-		return { &InnerBindTypes[0].Value, RangeBindings, NumRanges, InnermostSchema, Offset };
-	}
-};
-
-//////////////////////////////////////////////////////////////////////////
-
 // Save -> load struct ids for ESchemaFormat::InMemoryNames
 [[nodiscard]] PLAINPROPS_API TArray<FStructSchemaId> IndexInMemoryNames(const FSchemaBatch& Schemas,  FIdIndexerBase& Indexer);
 
diff --git a/Engine/Plugins/Experimental/PlainProps/Source/Public/PlainPropsBuild.h b/Engine/Plugins/Experimental/PlainProps/Source/Public/PlainPropsBuild.h
--- a/Engine/Plugins/Experimental/PlainProps/Source/Public/PlainPropsBuild.h
+++ b/Engine/Plugins/Experimental/PlainProps/Source/Public/PlainPropsBuild.h
@@ -132,23 +132,34 @@ inline uint64 ValueCast(char32_t Value)		{ return static_cast<uint32>(Value); }
 
 //////////////////////////////////////////////////////////////////////////
 
-// Todo: Turn into smart pointer where dtor call FBuiltRange::Delete if Values not detached to FBuiltValue
 struct FTypedRange
 {
 	FMemberSchema Schema;
 	FBuiltRange* Values = nullptr;
 };
 
+template<Arithmetic T>
+FMemberSchema MakeLeafRangeSchema(ERangeSizeType MaxSize)
+{
+	return { FMemberType(MaxSize), ReflectArithmetic<T>.Pack(), 1, NoId, nullptr };
+}
+
 template<Arithmetic T, typename SizeType>
 FMemberSchema MakeLeafRangeSchema()
 {
-	return { FMemberType(RangeSizeOf(SizeType{})), ReflectArithmetic<T>.Pack(), 1, NoId, nullptr };
+	return MakeLeafRangeSchema<T>(RangeSizeOf(SizeType{}));
+}
+
+template<Enumeration T>
+FMemberSchema MakeEnumRangeSchema(ERangeSizeType MaxSize, FEnumSchemaId Schema)
+{
+	return { FMemberType(MaxSize), ReflectEnum<T>.Pack(), 1, Schema, nullptr };
 }
 
 template<Enumeration T, typename SizeType>
 FMemberSchema MakeEnumRangeSchema(FEnumSchemaId Schema)
 {
-	return { FMemberType(RangeSizeOf(SizeType{})), ReflectEnum<T>.Pack(), 1, Schema, nullptr };
+	return MakeEnumRangeSchema<T>(RangeSizeOf(SizeType{}), Schema);
 }
 
 inline constexpr FMemberType DefaultStructType = FMemberType(FStructType{EMemberKind::Struct, /* IsDynamic */ 0, /* IsSuper */ 0});
@@ -162,13 +173,6 @@ inline FMemberSchema MakeStructRangeSchema(ERangeSizeType SizeType, FStructSchem
 
 PLAINPROPS_API FMemberSchema MakeNestedRangeSchema(FScratchAllocator& Scratch, ERangeSizeType SizeType, FMemberSchema InnerRangeSchema);
 
-// @param InnerTypes must outlive FMemberSchema
-template<uint16 N>
-inline FMemberSchema MakeNestedRangeSchema(ERangeSizeType SizeType, const FMemberType (&InnerTypes)[N], FOptionalSchemaId InnermostSchema)
-{
-	return { FMemberType(SizeType), InnerTypes[0], N, InnermostSchema, N > 1 ? InnerTypes : nullptr };
-}
-
 //////////////////////////////////////////////////////////////////////////
 
 [[nodiscard]] PLAINPROPS_API FBuiltRange* CloneLeaves(FScratchAllocator& Scratch, uint64 Num, const void* Data, SIZE_T LeafSize);
@@ -207,8 +211,15 @@ union FBuiltValue
 	FBuiltRange*	Range;
 };
 
+struct FTypedValue
+{
+	FMemberSchema	Schema;
+	FBuiltValue		Value;
+};
+
 struct FBuiltMember
 {
+	FBuiltMember(FMemberId InName, FTypedValue In) : FBuiltMember(InName, In.Schema, In.Value) {}
 	FBuiltMember(FMemberId Name, FUnpackedLeafType Leaf, FOptionalEnumSchemaId Schema, uint64 Value);
 	FBuiltMember(FMemberId Name, FTypedRange Range);
 	FBuiltMember(FMemberId Name, FStructSchemaId Schema, FBuiltStruct* Value);
@@ -246,6 +257,7 @@ public:
 	void AddLeaf(FMemberId Name, FUnpackedLeafType Leaf, FOptionalEnumSchemaId Enum, uint64 Value)	{ Members.Emplace(Name, Leaf, Enum, Value); }
 	void AddRange(FMemberId Name, FTypedRange Range)												{ Members.Emplace(Name, Range); }
 	void AddStruct(FMemberId Name, FStructSchemaId Schema, FBuiltStruct* Struct)					{ Members.Emplace(Name, Schema, Struct); }
+	void Add(FMemberId Name, FTypedValue TypedValue)												{ Members.Emplace(Name, TypedValue); }
 	
 	// Build members into a single nested super struct member, no-op if no non-super members has been added
 	PLAINPROPS_API void BuildSuperStruct(FScratchAllocator&	Scratch, const FStructDeclaration& Super, const FDebugIds& Debug);
diff --git a/Engine/Plugins/Experimental/PlainProps/Source/Public/PlainPropsLoad.h b/Engine/Plugins/Experimental/PlainProps/Source/Public/PlainPropsLoad.h
--- a/Engine/Plugins/Experimental/PlainProps/Source/Public/PlainPropsLoad.h
+++ b/Engine/Plugins/Experimental/PlainProps/Source/Public/PlainPropsLoad.h
@@ -2,8 +2,8 @@
 
 #pragma once
 
-#include "Containers/ArrayView.h"
 #include "PlainPropsRead.h"
+#include "Containers/ArrayView.h"
 #include "Memory/MemoryFwd.h"
 #include "Templates/UniquePtr.h"
 
@@ -12,10 +12,8 @@ namespace PlainProps
 	
 class FCustomBindings;
 class FDeclarations;
-struct FLoadBatch;
-class FRangeBinding;
 class FSchemaBindings;
-struct FTypedRange;
+struct FLoadBatch;
 
 struct FLoadBatchDeleter 
 {
@@ -24,17 +22,13 @@ struct FLoadBatchDeleter
 using FLoadBatchPtr = TUniquePtr<FLoadBatch, FLoadBatchDeleter>;
 
 template<class Runtime>
-[[nodiscard]] FLoadBatchPtr CreateLoadPlans(FReadBatchId ReadId, TConstArrayView<FStructSchemaId> RuntimeIds)
+[[nodiscard]] FLoadBatchPtr					CreateLoadPlans(FReadBatchId ReadId, TConstArrayView<FStructSchemaId> RuntimeIds)
 {
 	return CreateLoadPlans(ReadId, Runtime::GetTypes(), Runtime::GetCustoms(), Runtime::GetSchemas(), RuntimeIds);
 }
 
-PLAINPROPS_API FLoadBatchPtr				CreateLoadPlans(FReadBatchId ReadId, const FDeclarations& Declarations, const FCustomBindings& Customs, const FSchemaBindings& Schemas, TConstArrayView<FStructSchemaId> RuntimeIds);
+PLAINPROPS_API [[nodiscard]] FLoadBatchPtr	CreateLoadPlans(FReadBatchId ReadId, const FDeclarations& Declarations, const FCustomBindings& Customs, const FSchemaBindings& Schemas, TConstArrayView<FStructSchemaId> RuntimeIds);
 PLAINPROPS_API void							LoadStruct(void* Dst, FByteReader Src, FStructSchemaId Id, const FLoadBatch& Batch);
-PLAINPROPS_API void							LoadStruct(void* Dst, FStructView Src, const FLoadBatch& Batch);
-PLAINPROPS_API void							ConstructAndLoadStruct(void* Dst, FByteReader Src, FStructSchemaId Id, const FLoadBatch& Batch);
-PLAINPROPS_API void							ConstructAndLoadStruct(void* Dst, FStructView Src, const FLoadBatch& Batch);
-PLAINPROPS_API void							LoadRange(void* Dst, FRangeView Src, ERangeSizeType MaxType, TConstArrayView<FRangeBinding> Bindings, const FLoadBatch& Batch);
-
-} // namespace PlainProps
+PLAINPROPS_API void							ConstructAndLoadStruct(void* Dst, FByteReader Src, FStructSchemaId LoadId, const FLoadBatch& Batch);
 
+} // namespace PlainProps
\ No newline at end of file
diff --git a/Engine/Plugins/Experimental/PlainProps/Source/Public/PlainPropsLoadMember.h b/Engine/Plugins/Experimental/PlainProps/Source/Public/PlainPropsLoadMember.h
new file mode 100644
--- /dev/null
+++ b/Engine/Plugins/Experimental/PlainProps/Source/Public/PlainPropsLoadMember.h
@@ -0,0 +1,130 @@
+// Copyright Epic Games, Inc. All Rights Reserved.
+
+#pragma once
+
+#include "PlainPropsLoad.h"
+
+namespace PlainProps 
+{
+
+class FNestedRangeLoadIterator;
+class FRangeLoader;
+class FRangeBinding;
+class FStructRangeLoadIterator;
+
+struct FSchemaLoadHandle
+{
+	FStructSchemaId						LoadId;
+	const FLoadBatch&					Batch;
+};
+
+// Usable via FMemberLoader or [ConstructAnd]LoadStruct()
+struct FStructLoadView
+{
+	FSchemaLoadHandle					Schema;
+	FByteReader							Values;
+};
+
+PLAINPROPS_API void	LoadStruct(void* Dst, FStructLoadView Src);
+PLAINPROPS_API void	ConstructAndLoadStruct(void* Dst, FStructLoadView Src);
+
+//////////////////////////////////////////////////////////////////////////
+
+struct FRangeLoadSchema
+{
+	FMemberType						ItemType;
+	FOptionalSchemaId				InnermostId;
+	const FMemberType*				NestedItemTypes; // For nested ranges, can be out-of-bounds otherwise
+	const FLoadBatch&				Batch;
+};
+
+using FStructRangeLoadView = TStructuralRangeView<FStructRangeLoadIterator>;
+using FNestedRangeLoadView = TStructuralRangeView<FNestedRangeLoadIterator>;
+
+class FRangeLoadView
+{
+public:
+	FRangeLoadView(FRangeLoadSchema S, uint64 N, FMemoryView V) : Schema(S), NumItems(N), Values(V)  {}
+
+	uint64								Num() const				{ return NumItems;}
+	bool								IsEmpty() const			{ return NumItems == 0; }
+	bool								IsLeafRange() const		{ return Schema.ItemType.GetKind() == EMemberKind::Leaf; }
+	bool								IsStructRange() const	{ return Schema.ItemType.GetKind() == EMemberKind::Struct; }
+	bool								IsNestedRange() const	{ return Schema.ItemType.GetKind() == EMemberKind::Range; }
+
+	PLAINPROPS_API FLeafRangeView		AsLeaves() const;		// @pre IsLeafRange()
+	PLAINPROPS_API FStructRangeLoadView	AsStructs() const;		// @pre IsStructRange()
+	PLAINPROPS_API FNestedRangeLoadView	AsRanges() const;		// @pre IsNestedRange()
+
+private:
+	friend FRangeLoader;
+
+	FRangeLoadSchema				Schema;
+	uint64							NumItems;
+	FMemoryView						Values;
+};
+
+PLAINPROPS_API void	LoadRange(void* Dst, FRangeLoadView Src, TConstArrayView<FRangeBinding> InnerBindings);
+
+//////////////////////////////////////////////////////////////////////////
+
+class FNestedRangeLoadIterator
+{
+	friend FNestedRangeLoadView;
+	FRangeLoadSchema		Schema;
+	FByteReader				ByteIt;
+	FBitCacheReader			BitIt;
+
+public:
+	FNestedRangeLoadIterator(const FRangeLoadSchema& InSchema, FMemoryView Data)
+	: Schema(InSchema) 
+	, ByteIt(Data)
+	{}
+
+	PLAINPROPS_API FRangeLoadView	operator*() const;
+	bool							operator!=(const FNestedRangeLoadIterator& Rhs) const { return ByteIt.Peek() != Rhs.ByteIt.Peek(); };
+	PLAINPROPS_API void				operator++();
+};
+
+class FStructRangeLoadIterator
+{
+	friend FStructRangeLoadView;
+	FSchemaLoadHandle		Schema;
+	FByteReader				ByteIt;
+
+public:
+	FStructRangeLoadIterator(const FSchemaLoadHandle& InSchema, FMemoryView Data)
+	: Schema(InSchema) 
+	, ByteIt(Data)
+	{}
+
+	FStructLoadView operator*() const { return { Schema, FByteReader(ByteIt.PeekSkippableSlice()) }; }
+	bool operator!=(const FStructRangeLoadIterator& Rhs) const { return ByteIt.Peek() != Rhs.ByteIt.Peek(); };
+	void operator++() { (void)ByteIt.GrabSkippableSlice(); }
+};
+
+//////////////////////////////////////////////////////////////////////////
+
+class FMemberLoader
+{
+public:
+	explicit FMemberLoader(FStructLoadView Struct);
+
+	bool					HasMore() const				{ return Reader.HasMore(); }
+	
+	FOptionalMemberId		PeekName() const			{ return Reader.PeekName(); }
+	FOptionalMemberId		PeekNameUnchecked() const	{ return Reader.PeekNameUnchecked(); }
+	EMemberKind				PeekKind() const			{ return Reader.PeekKind(); }
+	FMemberType				PeekType() const			{ return Reader.PeekType(); }
+
+	FLeafView				GrabLeaf()					{ return Reader.GrabLeaf(); }
+	FRangeLoadView			GrabRange();				// @pre PeekKind() == EMemberKind::Range
+	FStructLoadView			GrabStruct();				// @pre PeekKind() == EMemberKind::Struct
+
+private:
+	FMemberReader			Reader;
+	const FStructSchemaId*	LoadIdIt;
+	const FLoadBatch&		Batch;
+};
+
+} // namespace PlainProps
\ No newline at end of file
diff --git a/Engine/Plugins/Experimental/PlainProps/Source/Public/PlainPropsRead.h b/Engine/Plugins/Experimental/PlainProps/Source/Public/PlainPropsRead.h
--- a/Engine/Plugins/Experimental/PlainProps/Source/Public/PlainPropsRead.h
+++ b/Engine/Plugins/Experimental/PlainProps/Source/Public/PlainPropsRead.h
@@ -13,6 +13,7 @@ namespace PlainProps
 {
 
 struct FEnumSchema;
+class FMemberLoader;
 struct FSchemaBatch;
 struct FStructSchema;
 
@@ -180,7 +181,7 @@ struct FLeafView
 {
 	FUnpackedLeafType		Leaf;
 	FReadBatchId			Batch;
-	FEnumSchemaId			Enum;
+	FOptionalEnumSchemaId	Enum;
 	FMemberValue			Value;
 
 	FORCEINLINE bool		AsBool() const		{ return Value.bValue; }
@@ -202,18 +203,10 @@ struct FLeafView
 	FORCEINLINE uint32		AsEnum32() const	{ return As<uint32,	FUnpackedLeafType{ELeafType::Enum, ELeafWidth::B32}>(); }
 	FORCEINLINE uint64		AsEnum64() const	{ return As<uint64,	FUnpackedLeafType{ELeafType::Enum, ELeafWidth::B64}>(); }
 
-private:
-	template<Arithmetic T, FUnpackedLeafType ExpectedLeaf = ReflectArithmetic<T>>
-	FORCEINLINE T As() const
-	{
-		check(ExpectedLeaf == Leaf);
-		return *reinterpret_cast<const T*>(Value.Ptr);
-	}
-
-	template<Enumeration T, FUnpackedLeafType ExpectedLeaf = ReflectEnum<T>>
-	FORCEINLINE T As() const
+	template<typename T, FUnpackedLeafType ExpectedLeaf = ReflectLeaf<T>>
+	FORCEINLINE T			As() const
 	{
-		check(ExpectedLeaf == Leaf);
+		check(Leaf == ExpectedLeaf);
 		return *reinterpret_cast<const T*>(Value.Ptr);
 	}
 };
@@ -258,9 +251,8 @@ private:
 //////////////////////////////////////////////////////////////////////////
 
 class FMemberReader;
-class FRangeLoader;
-class FStructRangeIterator;
 class FNestedRangeIterator;
+class FStructRangeIterator;
 
 template<typename IteratorType>
 class TStructuralRangeView;
@@ -291,9 +283,9 @@ public:
 	PLAINPROPS_API FNestedRangeView	AsRanges() const;		// @pre IsNestedRange()
 
 private:
+	friend FMemberLoader;
 	friend FMemberReader;
 	friend FNestedRangeIterator;
-	friend FRangeLoader;
 
 	FRangeSchema					Schema;
 	uint64							NumItems;
@@ -405,22 +397,40 @@ bool EqualItems(RangeTypeA&& A, RangeTypeB&& B)
 class FLeafRangeView
 {
 public:
+	FLeafRangeView(FUnpackedLeafType Leaf, FReadBatchId InBatch, FOptionalEnumSchemaId InEnum, uint64 Num, const uint8* Data) 
+	: Type(Leaf.Type)
+	, Width(Leaf.Width)
+	, Batch(InBatch)
+	, Enum(InEnum)
+	, NumItems(Num)
+	, Values(Data)
+	{}
+
 	ELeafType				GetType() const;
 	SIZE_T					GetSizeOf() const;
 	uint64					Num() const			{ return NumItems; }
 	uint64					NumBytes() const	{ return NumItems * GetSizeOf(); }
-	
-	// These ranges hide the internal representations so we can change format in the future, 
-	// e.g. store zeroes or 1.0f in some compact fashion or even var int encodings
+
+	// These range views hide the internal representations to enable future format changes, 
+	// e.g. store zeroes or 1.0f in some compact fashion or even variable length int encodings
 	//
 	// They could also provide various conversion helpers
+	
+	template<typename T>
+	auto					As() const
+	{
+		check(FUnpackedLeafType(Type, Width) == ReflectLeaf<T>);
+		return TRangeView<T>(reinterpret_cast<const T*>(Values), NumItems);	
+	}
+
+	template<> auto			As<bool>() const	{ return AsBools(); }
 
 	FBoolRangeView			AsBools() const
 	{
 		check(Type == ELeafType::Bool);
 		return FBoolRangeView(Values, NumItems);
 	}
-	
+
 	TRangeView< int8>		AsS8s() const		{ return As< int8>(); }
 	TRangeView<uint8>		AsU8s() const		{ return As<uint8>(); }
 	TRangeView< int16>		AsS16s() const		{ return As< int16>(); }
@@ -434,52 +444,8 @@ public:
 	TRangeView<char8_t>		AsUtf8() const		{ return As<char8_t>(); }
 	TRangeView<char16_t>	AsUtf16() const		{ return As<char16_t>(); }
 	TRangeView<char32_t>	AsUtf32() const		{ return As<char32_t>(); }
-	
-	template<Arithmetic T>
-	TRangeView<T> As() const
-	{
-		static_assert(!std::is_same_v<T, bool>);
-		check(FUnpackedLeafType(Type, Width) == ReflectArithmetic<T>);
-		return TRangeView<T>(reinterpret_cast<const T*>(Values), NumItems);
-	}
-
-	template<Enumeration T>
-	TRangeView<T> As() const
-	{
-		check(FUnpackedLeafType(Type, Width) == ReflectEnum<T>);
-		return TRangeView<T>(reinterpret_cast<const T*>(Values), NumItems);
-	}
-
-	//FEnumFlatRangeView		AsFlatEnums();	// @pre GetType() == ELeafType::Enum
-	//FEnumFlagRangeView		AsFlagEnums();	// @pre GetType() == ELeafType::Enum && 
-
-	//// Option B Memcopying + copy conversion apis
-	//
-	//void					ConvertBits(TArrayView<bool> Out);
-	//void					ConvertBits(TArrayView<uint8> Out);
-	//void					ConvertBits(TArrayView<uint16> Out);
-	//void					ConvertBits(TArrayView<uint32> Out);
-	//void					ConvertBits(TArrayView<uint64> Out);
-
-	//template<typename T>
-	//void					CopyExact(TArrayView<T> Out);
-
-	//// E.g. IntCastChecked, double -> float ok but not float -> double if it loses precision
-	////						uint64 -> uint32 ok if all 
-	//template<typename T>
-	//bool					TryConvertLossless(TArrayView<uint64> Out);
-
-	////  
-	//template<typename T>
-	//bool					TryConvert(TArrayView<T> Out);
-
-	//template<typename T>
-	//void					CopyStaticCast(TArrayView<T> Out);
-
 
 private:
-	friend class FRangeView;
-
 	ELeafType						Type;
 	ELeafWidth						Width;
 	FReadBatchId					Batch;
@@ -530,6 +496,7 @@ class TStructuralRangeView
 {
 public:
 	using SchemaType = decltype(IteratorType::Schema);
+	TStructuralRangeView(uint64 N, FMemoryView D, SchemaType S) : NumItems(N), Data(D), Schema(S) {}
 
 	uint64			Num() const		{ return NumItems; }
 	SchemaType		GetSchema() const { return Schema; }
@@ -538,8 +505,6 @@ public:
 	IteratorType	end() const		{ return IteratorType(Schema, FMemoryView(Data.GetDataEnd(), 0)); }
 
 private:
-	friend class FRangeView;
-
 	uint64				NumItems;
 	FMemoryView			Data;
 	SchemaType			Schema;
@@ -568,6 +533,8 @@ private:
 //	FMemberValue		Value;
 //};
 
+//////////////////////////////////////////////////////////////////////////
+
 // Iterates over struct members
 class FMemberReader
 {
diff --git a/Engine/Plugins/Experimental/PlainProps/Source/Public/PlainPropsSave.h b/Engine/Plugins/Experimental/PlainProps/Source/Public/PlainPropsSave.h
--- a/Engine/Plugins/Experimental/PlainProps/Source/Public/PlainPropsSave.h
+++ b/Engine/Plugins/Experimental/PlainProps/Source/Public/PlainPropsSave.h
@@ -3,7 +3,6 @@
 #pragma once
 
 #include "PlainPropsTypes.h"
-#include "Templates/UniquePtr.h"
 
 namespace PlainProps 
 {
@@ -12,6 +11,7 @@ struct FBuiltRange;
 struct FBuiltStruct;
 class FCustomBindings;
 class FDeclarations;
+struct ILeafRangeBinding;
 struct FRangeMemberBinding;
 class FSchemaBindings;
 class FScratchAllocator;
@@ -36,7 +36,11 @@ FSaveContext MakeSaveContext(FDefaultStructs Defaults, FScratchAllocator& Scratc
 	return { Runtime::GetTypes(), Runtime::GetSchemas(), Runtime::GetCustoms(), Scratch, Defaults };
 }
 
+////////////////////////////////////////////////////////////////////////////////////////////////
+
 [[nodiscard]] PLAINPROPS_API FBuiltStruct*	SaveStruct(const void* Struct, FStructSchemaId BindId, const FSaveContext& Context);
 [[nodiscard]] PLAINPROPS_API FBuiltStruct*	SaveStructDelta(const void* Struct, const void* Default, FStructSchemaId BindId, const FSaveContext& Context);
 [[nodiscard]] PLAINPROPS_API FBuiltRange*	SaveRange(const void* Range, FRangeMemberBinding Member, const FSaveContext& Ctx);
+[[nodiscard]] PLAINPROPS_API FBuiltRange*	SaveLeafRange(const void* Range, const ILeafRangeBinding& Binding, FUnpackedLeafType Leaf, const FSaveContext& Ctx);
+
 } // namespace PlainProps
\ No newline at end of file
diff --git a/Engine/Plugins/Experimental/PlainProps/Source/Public/PlainPropsSaveMember.h b/Engine/Plugins/Experimental/PlainProps/Source/Public/PlainPropsSaveMember.h
new file mode 100644
--- /dev/null
+++ b/Engine/Plugins/Experimental/PlainProps/Source/Public/PlainPropsSaveMember.h
@@ -0,0 +1,326 @@
+// Copyright Epic Games, Inc. All Rights Reserved.
+
+#pragma once
+
+#include "PlainPropsTypes.h"
+#include "PlainPropsBind.h"
+#include "PlainPropsBuild.h"
+#include "PlainPropsLoadMember.h"
+#include "PlainPropsSave.h"
+
+namespace PlainProps 
+{
+
+// TMemberSerializer helper
+template<typename T>
+struct TStructSerializer
+{
+	static constexpr EMemberKind Kind = EMemberKind::Struct;
+
+	template<class Ids>
+	TStructSerializer(TCustomInit<Ids> Init, TConstArrayView<FMemberId> Names, FDualStructSchemaId Dual = GetStructDualId<Ids, T>())
+	: BindId(Dual.BindId)
+	{
+		Init.RegisterInnerStruct(Dual, Names);
+	}
+
+	FTypedValue SaveMember(const T& Value, const FSaveContext& Ctx) const
+	{
+		FMemberSchema Schema = { DefaultStructType, DefaultStructType, 0, BindId, nullptr};
+		return { Schema, {.Struct = SaveItem(Value, Ctx)} };
+	}
+
+	void LoadMember(T& Dst, /* in-out */ FMemberLoader& Src) const
+	{
+		LoadStruct(&Dst, Src.GrabStruct());
+	}
+
+	void ConstructAndLoadMember(T& Dst, /* in-out */ FMemberLoader& Src) const
+	{
+		if constexpr (std::is_default_constructible_v<T>)
+		{
+			LoadMember(*new (&Dst) T, Src);
+		}
+		else
+		{
+			static_assert(!std::is_void_v<CustomBind<T>>, "Non-default constructible types must load via custom bindings");
+			ConstructAndLoadStruct(&Dst, Src.GrabStruct());
+		}
+	}
+
+	////////// TRangeSaver API //////////
+	
+	using BuiltItemType = FBuiltStruct*;
+
+	FBuiltStruct* SaveItem(const T& Value, const FSaveContext& Ctx) const
+	{
+		return SaveStruct(&Value, BindId, Ctx);
+	}
+
+	FMemberSchema MakeMemberRangeSchema(ERangeSizeType MaxSize) const
+	{
+		return { FMemberType(MaxSize), DefaultStructType, 1, BindId, nullptr };
+	}
+
+	const FStructSchemaId BindId;
+};
+
+//////////////////////////////////////////////////////////////////////////
+
+template<typename T>
+inline constexpr FMemberType ReflectInnermostType()
+{
+	return DefaultStructType;
+}
+
+template<Arithmetic T>
+inline constexpr FMemberType ReflectInnermostType()
+{
+	return ReflectArithmetic<T>.Pack();
+}
+template<Enumeration T>
+inline constexpr FMemberType ReflectInnermostType()
+{
+	return ReflectEnum<T>.Pack();
+}
+
+// TRangeSerializer constructor helper, FMemberBindType isn't default-constructible 
+union FMemberBindTypeInitializer
+{
+	FMemberBindTypeInitializer() : Unused() {}
+	uint8_t			Unused;
+	FMemberBindType Value;
+};
+
+// TMemberSerializer helper
+template<typename T, typename RangeBinding>
+struct TRangeSerializer
+{
+	static constexpr EMemberKind	Kind = EMemberKind::Range;
+	static constexpr uint16			NumRanges = CountRangeBindings<RangeBinding>();
+	static constexpr uint16			NumTypes = NumRanges + 1;
+
+	FRangeBinding const* const		Bindings = nullptr;
+	FOptionalSchemaId				InnermostSchema;
+	FMemberBindTypeInitializer		BindTypes[NumTypes];
+	FMemberType						Types[NumTypes];
+
+	template<class Ids>
+	TRangeSerializer(TCustomInit<Ids> Init, TConstArrayView<FMemberId> Names)
+	: Bindings(GetRangeBindings<RangeBinding, NumRanges>().GetData())
+	{
+		for (uint16 Idx = 0; Idx < NumRanges; ++Idx)
+		{
+			ERangeSizeType MaxSize = Bindings[Idx].GetSizeType();
+			Types[Idx] = FMemberType(MaxSize);
+			BindTypes[Idx].Value = FMemberBindType(MaxSize);
+		}
+
+		using InnermostType = typename TInnermostType<RangeBinding, NumRanges>::Type;
+		Types[NumRanges] = ReflectInnermostType<InnermostType>();
+		BindTypes[NumRanges].Value = BindInnermostType<InnermostType, Ids>(/* out */ InnermostSchema);
+		
+
+		if constexpr (!LeafType<InnermostType>)
+		{			
+			FDualStructSchemaId Id = {static_cast<FStructSchemaId>(InnermostSchema.Get()), GetStructDeclId<Ids, InnermostType>()};
+			Init.RegisterInnerStruct(Id, Names);
+		}
+	}
+
+	void LoadMember(T& Dst, /* in-out */ FMemberLoader& Src) const
+	{
+		LoadRange(&Dst, Src.GrabRange(), MakeArrayView(Bindings, NumRanges));
+	}
+
+	FTypedValue	SaveMember(const T& Value, const FSaveContext& Ctx) const
+	{
+		return { MakeMemberSchema(), {.Range = SaveItem(Value, Ctx)} };
+	}
+		
+	FMemberSchema MakeMemberSchema() const
+	{
+		return { Types[0], Types[1], NumRanges, InnermostSchema, NumRanges > 1 ? Types + 1 : nullptr };
+	}
+
+	////////// TRangeSaver API //////////
+
+	using BuiltItemType = FBuiltRange*;
+
+	FMemberSchema MakeMemberRangeSchema(ERangeSizeType MaxSize) const
+	{
+		return { FMemberType(MaxSize), Types[0], NumTypes, InnermostSchema, Types };
+	}
+
+	FBuiltRange* SaveItem(const T& Item, const FSaveContext& Ctx) const
+	{
+		if constexpr (std::is_base_of_v<ILeafRangeBinding, RangeBinding>)
+		{
+			static_assert(NumRanges == 1);
+			return SaveLeafRange(&Item, Bindings[0].AsLeafBinding(), Types[1].AsLeaf(), Ctx);
+		}
+		else
+		{
+			FRangeMemberBinding MemberBinding = { &BindTypes[1].Value, Bindings, NumRanges, InnermostSchema, /* offset */ 0 };
+			return SaveRange(&Item, MemberBinding, Ctx);
+		}
+	}
+};
+
+//////////////////////////////////////////////////////////////////////////
+
+// TMemberSerializer helper
+template<typename T>
+struct TLeafSerializer
+{
+	void LoadMember(T& Dst, /* in-out */ FMemberLoader& Src) const
+	{
+		Dst = Src.GrabLeaf().As<T>();
+	}
+
+	void ConstructAndLoadMember(T& Dst, /* in-out */ FMemberLoader& Src) const
+	{
+		LoadMember(Dst, Src);
+	}
+
+	static FTypedValue SaveLeaf(FMemberType Type, FOptionalSchemaId Id, T Value)
+	{
+		return { {Type, Type, 0, Id, nullptr}, {.Leaf = ValueCast(Value)} };
+	}
+
+	////////// TRangeSaver API //////////
+	
+	using BuiltItemType = T;
+
+	static T SaveItem(T Value, const FSaveContext&) { return Value;}
+};
+
+// TMemberSerializer helper
+template<Arithmetic T>
+struct TArithmeticSerializer : TLeafSerializer<T>
+{
+	static constexpr FMemberType MemberType = ReflectArithmetic<T>.Pack();
+
+	template<class Ids>
+	TArithmeticSerializer(TCustomInit<Ids>, TConstArrayView<FMemberId>)
+	{}
+
+	static FTypedValue SaveMember(T Value, const FSaveContext&)
+	{
+		return SaveLeaf(MemberType, NoId, Value);
+	}
+
+	////////// TRangeSaver API //////////
+
+	FMemberSchema MakeMemberRangeSchema(ERangeSizeType MaxSize) const
+	{
+		return { FMemberType(MaxSize), MemberType, 1, NoId, nullptr };
+	}
+};
+
+// TMemberSerializer helper
+template<Enumeration T>
+struct TEnumSerializer : TLeafSerializer<T>
+{
+	static constexpr FMemberType MemberType = ReflectEnum<T>.Pack();
+	
+	const FEnumSchemaId Id;
+
+	template<class Ids>
+	TEnumSerializer(TCustomInit<Ids>, TConstArrayView<FMemberId>)
+	: Id(GetEnumId<Ids, T>())
+	{}
+
+	FTypedValue SaveMember(T Value, const FSaveContext& Ctx) const
+	{
+		return SaveLeaf(MemberType, Id, Value);
+	}
+
+	////////// TRangeSaver API //////////
+
+	FMemberSchema MakeMemberRangeSchema(ERangeSizeType MaxSize) const
+	{
+		return { FMemberType(MaxSize), MemberType, 1, Id, nullptr };
+	}
+};
+
+//////////////////////////////////////////////////////////////////////////
+
+template<typename T>
+struct TSelectSerializer
+{
+	using CustomBinding = CustomBind<T>;
+	using RangeBinding = RangeBind<T>;
+	static constexpr bool bRange = std::is_void_v<CustomBinding> && !std::is_void_v<RangeBinding>;
+	using Type = std::conditional_t<bRange, TRangeSerializer<T, RangeBinding>, TStructSerializer<T>>;
+};
+
+template<Arithmetic T>
+struct TSelectSerializer<T>
+{
+	using Type = TArithmeticSerializer<T>;
+};
+
+template<Enumeration T>
+struct TSelectSerializer<T>
+{
+	using Type = TEnumSerializer<T>;
+};
+
+// Helps templated custom bindings save generic members
+template<typename T>
+using TMemberSerializer = typename TSelectSerializer<T>::Type;
+
+//////////////////////////////////////////////////////////////////////////
+
+class FRangeSaverBase
+{
+protected:
+	PLAINPROPS_API FRangeSaverBase(FScratchAllocator& Scratch, uint64 Num, SIZE_T ItemSize);
+	~FRangeSaverBase() { check(!Range);	}
+
+	FBuiltRange*				Range;
+	uint8*						It;
+#if DO_CHECK
+	const uint8*				End;
+#endif
+};
+
+// Saves a range of T without a range binding
+template<typename T>
+class TRangeSaver : public FRangeSaverBase
+{
+public:
+	using FMemberSerializer = TMemberSerializer<T>;
+	using BuiltItemType = typename FMemberSerializer::BuiltItemType;
+	static constexpr SIZE_T ItemSize = sizeof(BuiltItemType);
+
+	TRangeSaver(const FSaveContext& InCtx, uint64 Num, const FMemberSerializer& InSchema)
+	: FRangeSaverBase(InCtx.Scratch, Num, sizeof(BuiltItemType))
+	, Schema(InSchema)
+	, Ctx(InCtx)
+	{}
+
+	void AddItem(const T& Item)
+	{
+		check(Range);
+		check(It < End);
+		reinterpret_cast<BuiltItemType&>(*It) = Schema.SaveItem(Item, Ctx);
+		It += ItemSize;
+	}
+
+	[[nodiscard]] FTypedRange Finalize(ERangeSizeType MaxSize)
+	{
+		check(Range);
+		check(It == End);
+		FTypedRange Out { Schema.MakeMemberRangeSchema(MaxSize), Range };
+		Range = nullptr;
+		return Out;
+	}
+
+private:	
+	const FMemberSerializer&	Schema;
+	const FSaveContext&			Ctx;
+};
+
+} // namespace PlainProps
\ No newline at end of file
diff --git a/Engine/Plugins/Experimental/PlainProps/Source/Public/PlainPropsTypes.h b/Engine/Plugins/Experimental/PlainProps/Source/Public/PlainPropsTypes.h
--- a/Engine/Plugins/Experimental/PlainProps/Source/Public/PlainPropsTypes.h
+++ b/Engine/Plugins/Experimental/PlainProps/Source/Public/PlainPropsTypes.h
@@ -378,6 +378,16 @@ template<> inline constexpr FUnpackedLeafType ReflectArithmetic<char8_t>	= { ELe
 template<> inline constexpr FUnpackedLeafType ReflectArithmetic<char16_t>	= { ELeafType::Unicode,	ELeafWidth::B16 };
 template<> inline constexpr FUnpackedLeafType ReflectArithmetic<char32_t>	= { ELeafType::Unicode,	ELeafWidth::B32 };
 
+
+template<Arithmetic T>
+inline constexpr FUnpackedLeafType MakeLeafType() { return ReflectArithmetic<T>; }
+
+template<Enumeration T>
+inline constexpr FUnpackedLeafType MakeLeafType() { return ReflectEnum<T>; }
+
+template<typename T>
+inline constexpr FUnpackedLeafType ReflectLeaf = MakeLeafType<T>();
+
 //////////////////////////////////////////////////////////////////////////
 
 inline constexpr ERangeSizeType RangeSizeOf(bool)	{ return ERangeSizeType::Uni; }
diff --git a/Engine/Plugins/Experimental/PlainProps/Source/Public/PlainPropsUeCoreBindings.h b/Engine/Plugins/Experimental/PlainProps/Source/Public/PlainPropsUeCoreBindings.h
--- a/Engine/Plugins/Experimental/PlainProps/Source/Public/PlainPropsUeCoreBindings.h
+++ b/Engine/Plugins/Experimental/PlainProps/Source/Public/PlainPropsUeCoreBindings.h
@@ -8,14 +8,12 @@
 #include "Math/Transform.h"
 #include "Misc/Optional.h"
 #include "PlainPropsBind.h"
-#include "PlainPropsBuild.h" // FMemberSchema for TSetDeltaBinding::SaveSet
 #include "PlainPropsIndex.h"
-#include "PlainPropsLoad.h"
-#include "PlainPropsRead.h"
+#include "PlainPropsLoadMember.h"
+#include "PlainPropsSaveMember.h"
 #include "PlainPropsStringUtil.h"
 #include "UObject/NameTypes.h"
 
-
 PP_REFLECT_STRUCT_TEMPLATE(, TTuple, void, Key, Value); // Todo handle TTuple and higher arities
 
 namespace UE::Math
@@ -420,6 +418,50 @@ struct TSetBinding : public IItemRangeBinding
 
 //////////////////////////////////////////////////////////////////////////
 
+// Only used for non-default constructible pairs
+template <typename K, typename V>
+struct TPairBinding : ICustomBinding
+{
+	using Type = TPair<K,V>;
+
+	template<class Ids>
+	TPairBinding(TCustomInit<Ids> Init)
+	: MemberIds{Ids::IndexMember("Key"), Ids::IndexMember("Value")}
+	, Key(Init, { MemberIds[0] })
+	, Value(Init, { MemberIds[1] })
+	{}
+
+	void Save(FMemberBuilder& Dst, const TPair<K,V>& Src, const Type* Default, const FSaveContext& Ctx) const
+	{
+		Dst.Add(MemberIds[0], Key.SaveMember(Src.Key, Ctx));
+		Dst.Add(MemberIds[1], Value.SaveMember(Src.Value, Ctx));
+	}
+		
+	void Load(TPair<K,V>& Dst, FStructLoadView Src, ECustomLoadMethod Method) const
+	{
+		FMemberLoader Members(Src);
+		check(Members.PeekName() == MemberIds[0]);
+		if (Method == ECustomLoadMethod::Construct)
+		{
+			Key.ConstructAndLoadMember(/* out */ Dst.Key, /* in-out */ Members);
+			Value.ConstructAndLoadMember(/* out */ Dst.Value, /* in-out */ Members);		
+		}
+		else
+		{
+			Key.LoadMember(/* out */ Dst.Key, /* in-out */ Members);
+			Value.LoadMember(/* out */ Dst.Value, /* in-out */ Members);				
+		}
+	}
+	
+	static bool	Diff(const TPair<K,V>& A, const TPair<K,V>& B) { return A != B; }
+
+	const FMemberId				MemberIds[2];
+	const TMemberSerializer<K>	Key;
+	const TMemberSerializer<V>	Value;
+};
+
+//////////////////////////////////////////////////////////////////////////
+
 template <typename K, typename V, typename SetAllocator, typename KeyFuncs>
 struct TMapBinding : public TSetBinding<TPair<K, V>, KeyFuncs, SetAllocator>
 {};
@@ -430,220 +472,230 @@ struct TMapBinding : public TSetBinding<TPair<K, V>, KeyFuncs, SetAllocator>
 struct FTransformBinding : ICustomBinding
 {
 	using Type = FTransform;
-	inline static constexpr EMemberPresence Occupancy = EMemberPresence::AllowSparse;
 	enum class EMember : uint8 { Translate, Rotate, Scale };
 	const FMemberId MemberIds[3];
 	const FStructSchemaId VectorId;
 	const FStructSchemaId QuatId;
 
 	template<class Ids>
-	FTransformBinding(TCustomInit<Ids>/*, const FDeclarations& Declared*/)
+	FTransformBinding(TCustomInit<Ids>)
 	: MemberIds{Ids::IndexMember("Translate"), Ids::IndexMember("Rotate"), Ids::IndexMember("Scale")}
 	, VectorId(GetStructDeclId<Ids, FVector>())
 	, QuatId(GetStructDeclId<Ids, FQuat>())
 	{}
 
 	PLAINPROPS_API void	Save(FMemberBuilder& Dst, const FTransform& Src, const FTransform* Default, const FSaveContext& Context) const;
-	PLAINPROPS_API void	Load(FTransform& Dst, FStructView Src, ECustomLoadMethod Method, const FLoadBatch& Batch) const;
+	PLAINPROPS_API void	Load(FTransform& Dst, FStructLoadView Src, ECustomLoadMethod Method) const;
 	inline static bool	Diff(const FTransform& A, const FTransform& B) { return !A.Equals(B, 0.0); }
 };
 
 //////////////////////////////////////////////////////////////////////////
 
-struct FSetDeltaIds
+struct FBaseDeltaBinding
 {
-	enum class EMember : uint8 { Del, Add };
-	const FMemberId MemberIds[2];
+	enum class EOp { Del, Add };
+	FMemberId MemberIds[2];
+	FMemberId DelId() const { return MemberIds[(uint32)EOp::Del]; }
+	FMemberId AddId() const { return MemberIds[(uint32)EOp::Add]; }
 	
 	template<class Ids>	
-	FSetDeltaIds(TCustomInit<Ids>)
+	FBaseDeltaBinding(TInit<Ids>)
 	: MemberIds{Ids::IndexMember("Del"), Ids::IndexMember("Add")}
 	{}
 
 	template<class Ids>	
-	static FSetDeltaIds Cache()
+	static FBaseDeltaBinding Cache()
 	{
-		static FSetDeltaIds Out(TCustomInit<Ids>{});
+		static FBaseDeltaBinding Out(TInit<Ids>{});
 		return Out;
 	}
 };
 
-template <typename T, typename KeyFuncs, typename SetAllocator>
-struct TSetDeltaBinding : ICustomBinding, FSetDeltaIds
+template <typename SetType, typename KeyType>
+struct TBaseDeltaBinding : ICustomBinding, FBaseDeltaBinding
 {
-	using Type = TSet<T, KeyFuncs, SetAllocator>;
-	using FSet = Type;
-	using FSetBinding = TSetBinding<T, KeyFuncs, SetAllocator>;
-	using FRangeMemberHelper = TRangeMemberHelper<FSetBinding>;
-	static constexpr EMemberPresence Occupancy = EMemberPresence::AllowSparse;
-	static constexpr uint16 NumInnerRanges = FRangeMemberHelper::NumRanges - 1;
-
-	struct FCustomTypename
-	{
-		inline static constexpr std::string_view DeclName = "SetDelta";
-		inline static constexpr std::string_view BindName = Concat<DeclName, ShortTypename<KeyFuncs>, ShortTypename<SetAllocator>>;
-		inline static constexpr std::string_view Namespace;
-		using Parameters = std::tuple<T>;
-	};
-
-	FRangeMemberHelper InnerRange;
+	using Type = SetType;
+	using ElemType = typename SetType::ElementType;
+	static constexpr ERangeSizeType MaxSize = ERangeSizeType::S32;
+	static constexpr bool bIsSet = std::is_same_v<KeyType, ElemType>;
 	
+	const TMemberSerializer<ElemType> Elems;
+
 	template<class Ids>	
-	TSetDeltaBinding(TCustomInit<Ids>)
-	: FSetDeltaIds(FSetDeltaIds::Cache<Ids>())
-	{
-		InnerRange.template Init<Ids>();
-	}
-	
-	void Save(FMemberBuilder& Dst, const FSet& Src, const FSet* Default, const FSaveContext& Context) const
+	TBaseDeltaBinding(TCustomInit<Ids> Init)
+	: FBaseDeltaBinding(FBaseDeltaBinding::Cache<Ids>())
+	, Elems(Init, bIsSet ? MakeArrayView(MemberIds) : MakeArrayView({AddId()}))
+	{}
+
+	template<typename KeyOrElemType>
+	static const KeyOrElemType& Get(const ElemType& Elem)
 	{
-		if (Default && !Default->IsEmpty())
+		if constexpr (std::is_same_v<KeyOrElemType, ElemType>)
 		{
-			// Inefficient, mimic FSetProperty::SerializeItem for better perf
-			SaveSet(Dst, MemberIds[(uint8)EMember::Del], Default->Difference(Src), Context);
-			SaveSet(Dst, MemberIds[(uint8)EMember::Add], Src.Difference(*Default), Context);
+			return Elem;
 		}
-		else if (Src.Num())
+		else
 		{
-			SaveSet(Dst, MemberIds[(uint8)EMember::Add], Src, Context);
+			return Elem.Key;
 		}
 	}
 
-	void SaveSet(FMemberBuilder& Dst, FMemberId Name, const FSet& Set, const FSaveContext& Ctx) const
+	inline void SaveDelta(FMemberBuilder& Dst, const SetType& Src, const SetType* Default, const FSaveContext& Ctx, const TMemberSerializer<KeyType>& Keys) const
 	{
-		if (int32 Num = Set.Num())
+		if (Src.IsEmpty())
 		{
-			// Start of a more optimized version that saves typed values directly w/o leaning on TSetBinding
-			//
-			//FBuiltRange* Range;
-			//if constexpr (LeafType<T>)
-			//{
-			//	Range = FBuiltRange::Create(Ctx.Scratch, Num, sizeof(T));
-			//	T* OutIt = static_cast<T*>(Range->Data);
-			//	//const TSparseArray<TSetElement<T>& Elements = reinterpret_cast<const TSparseArray<TSetElement<T>>&>(Values);
-			//	//if (Elements.IsCompact())
-			//	//{
-			//	//	FMemory::Memcpy(Range->Data, Elements.GetData(), sizeof(T) * Num);
-			//	//}
-			//	//else 
-			//	for (T Leaf : Set)
-			//	{
-			//		*OutIt++ = Leaf;
-			//	}
-			//}
-			//else 
-			//{
-			//	Range = SaveRange(&Set, InnerRange.MakeBinding(/* offset*/ 0), Ctx);
-			//}
+			if (Default && !Default->IsEmpty())
+			{
+				Dst.AddRange(DelId(), SaveAll<KeyType>(*Default, Keys, Ctx));
+			}
+		}
+		else if (Default && !Default->IsEmpty())
+		{
+			TBitArray<> DelSubset(false, Default->GetMaxIndex());
+			for (auto It = Default->CreateConstIterator(); It; ++It)
+			{
+				DelSubset[It.GetId().AsInteger()] = !Src.Contains(Get<KeyType>(*It));
+			}
+			if (DelSubset.Find(true) != INDEX_NONE)
+			{
+				Dst.AddRange(DelId(), SaveSome<KeyType>(*Default, DelSubset, Keys, Ctx));
+			}
 
-			// Lean on TSetBinding for now, less efficient but simpler
-			FBuiltRange* Range = SaveRange(&Set, InnerRange.MakeBinding(/* offset*/ 0), Ctx);
-			FMemberSchema Schema = MakeNestedRangeSchema(InnerRange.MaxSize, InnerRange.InnerSchemaTypes, InnerRange.InnermostSchema);
-			Dst.AddRange(Name, { Schema, Range });	
+			TBitArray<> AddSubset(false, Src.GetMaxIndex());
+			for (auto It = Src.CreateConstIterator(); It; ++It)
+			{
+				AddSubset[It.GetId().AsInteger()] = !ContainsValue(*Default, *It);
+			}
+			if (AddSubset.Find(true) != INDEX_NONE)
+			{
+				Dst.AddRange(AddId(), SaveSome<ElemType>(Src, AddSubset, Elems, Ctx));
+			}
+		}
+		else
+		{
+			Dst.AddRange(AddId(), SaveAll<ElemType>(Src, Elems, Ctx));
 		}
+	}
 
+	template<typename ItemType>
+	FTypedRange SaveAll(const SetType& Set, const TMemberSerializer<ItemType>& Schema, const FSaveContext& Ctx) const
+	{
+		check(!Set.IsEmpty());
+
+		TRangeSaver<ItemType> Items(Ctx, static_cast<uint64>(Set.Num()), Schema);
+		for (const ElemType& Elem : Set)
+		{
+			Items.AddItem(Get<ItemType>(Elem));
+		}
+		return Items.Finalize(MaxSize);
+	}
+	
+	template<typename ItemType>
+	FTypedRange SaveSome(const SetType& Set, const TBitArray<>& Subset, const TMemberSerializer<ItemType>& Schema, const FSaveContext& Ctx) const
+	{
+		TRangeSaver<ItemType> Items(Ctx, Subset.CountSetBits(), Schema);
+		for (int32 Idx = 0, Max = Set.GetMaxIndex(); Idx < Max; ++Idx)
+		{
+			if (Subset[Idx])
+			{
+				const ElemType& Elem = Set.Get(FSetElementId::FromInteger(Idx));
+				Items.AddItem(Get<ItemType>(Elem));
+			}
+		}
+		return Items.Finalize(MaxSize);
 	}
 
-	inline void Load(FSet& Dst, FStructView Src, ECustomLoadMethod Method, const FLoadBatch& Batch) const
+	inline void LoadDelta(Type& Dst, FStructLoadView Src, ECustomLoadMethod Method, const TMemberSerializer<KeyType>& Keys) const
 	{
-		FMemberReader Members(Src);
+		FMemberLoader Members(Src);
 
 		if (Method == ECustomLoadMethod::Construct)
 		{
-			::new (&Dst) FSet;
+			::new (&Dst) SetType;
 		}
 				
-		if (Members.HasMore())
+		while (Members.HasMore())
 		{
-			if (Members.PeekName() == MemberIds[(uint8)EMember::Add])
+			if (Members.PeekName() == AddId())
 			{
-				ApplyItems<EMember::Add>(Dst, Members.GrabRange(), Batch);
-			}
-			else
-			{
-				check(Members.PeekName() == MemberIds[(uint8)EMember::Del]);
-				ApplyItems<EMember::Del>(Dst, Members.GrabRange(), Batch);
-		
-				if (Members.HasMore())
-				{
-					check(Members.PeekName() == MemberIds[(uint8)EMember::Add]);
-					ApplyItems<EMember::Add>(Dst, Members.GrabRange(), Batch);
-				}
+				ApplyItems<EOp::Add, ElemType>(Dst, Members.GrabRange(), Elems);
+				check(!Members.HasMore());
+				break;
 			}
 
-			check(!Members.HasMore());
+			check(Members.PeekName() == DelId());
+			ApplyItems<EOp::Del, KeyType>(Dst, Members.GrabRange(), Keys);
 		}
 	}
 
-	template<EMember Op>
-	void ApplyItems(FSet& Out, FRangeView Items, const FLoadBatch& Batch) const
+	template<EOp Op, typename T>
+	void ApplyItems(SetType& Out, FRangeLoadView Items, const TMemberSerializer<T>& Schema) const
 	{
 		check(!Items.IsEmpty());
-
-		if constexpr (Op == EMember::Add && !LeafType<T>)
+ 
+		if constexpr (Op == EOp::Add && !LeafType<T>)
 		{
 			Out.Reserve(static_cast<int32>(Items.Num()));
 		}
 
 		if constexpr (LeafType<T>)
 		{
-			ApplyLeaves<Op>(Out, Items.AsLeaves().As<T>());
+			ApplyLeaves<Op, T>(Out, Items.AsLeaves());
 		}
-		else if constexpr (NumInnerRanges)
+		else if constexpr (TMemberSerializer<T>::Kind == EMemberKind::Struct)
 		{
-			ApplyRanges<Op>(Out, Items.AsRanges(), Batch);
+			ApplyStructs<Op, T>(Out, Items.AsStructs());
 		}
-		else 
+		else
 		{
-			ApplyStructs<Op>(Out, Items.AsStructs(), Batch);
+			ApplyRanges<Op, T>(Out, Items.AsRanges(), Schema);
 		}
 	}
-	
-	template<EMember Op>
-	void ApplyLeaves(FSet& Out, TRangeView<T> Items) const
+
+	template<EOp Op, typename T>
+	void ApplyLeaves(SetType& Out, FLeafRangeView Items) const requires (Op == EOp::Add && !std::is_same_v<T, bool>)
 	{
-		if constexpr (Op == EMember::Add)
-		{
-			Out.Append(MakeArrayView(Items));
-		}
-		else
+		Out.Append(MakeArrayView(Items.As<T>()));
+	}
+
+	template<EOp Op, typename T>
+	void ApplyLeaves(SetType& Out, FLeafRangeView Items) const
+	{
+		for (T Item : Items.As<T>())
 		{
-			for (T Item : Items)
-			{
-				Out.Remove(Item);
-			}
+			ApplyItem<Op>(Out, Item);
 		}
 	}
 
-	template<EMember Op>
-	void ApplyRanges(FSet& Out, FNestedRangeView Items, const FLoadBatch& Batch) const
+	template<EOp Op, typename T>
+	void ApplyRanges(SetType& Out, FNestedRangeLoadView Items, const TMemberSerializer<T>& Schema) const
 	{
 		static_assert(std::is_default_constructible_v<T>, TEXT("Ranges must be default-constructible"));
 
-		TConstArrayView<FRangeBinding> InnerRangeBindings(InnerRange.RangeBindings + 1, NumInnerRanges);
-		for (FRangeView Item : Items)
+		TConstArrayView<FRangeBinding> Bindings(Schema.Bindings, Schema.NumRanges);
+		for (FRangeLoadView Item : Items)
 		{
 			T Tmp;
-			LoadRange(&Tmp, Item, InnerRange.MaxSize, InnerRangeBindings, Batch);
+			LoadRange(&Tmp, Item, Bindings);
 			ApplyItem<Op>(Out, MoveTemp(Tmp));
 		}
 	}
 	
-	template<EMember Op>
-	void ApplyStructs(FSet& Out, FStructRangeView Items, const FLoadBatch& Batch) const
+	template<EOp Op, typename T>
+	void ApplyStructs(SetType& Out, FStructRangeLoadView Items) const
 	{
-		for (FStructView Item : Items)
+		for (FStructLoadView Item : Items)
 		{
 			if constexpr (std::is_default_constructible_v<T>)
 			{
 				T Tmp;
-				LoadStruct(&Tmp, Item, Batch);
+				LoadStruct(&Tmp, Item);
 				ApplyItem<Op>(Out, MoveTemp(Tmp));
 			}
 			else
 			{
 				alignas(T) uint8 Buffer[sizeof(T)];
-				ConstructAndLoadStruct(Buffer, Item, Batch);
+				ConstructAndLoadStruct(Buffer, Item);
 				T& Tmp = *reinterpret_cast<T*>(Buffer);
 				ApplyItem<Op>(Out, MoveTemp(Tmp));
 				Tmp.~T();
@@ -651,12 +703,12 @@ struct TSetDeltaBinding : ICustomBinding, FSetDeltaIds
 		}
 	}
 	
-	template<EMember Op>
-	void ApplyItem(FSet& Out, T&& Item) const
+	template<EOp Op, typename T>
+	void ApplyItem(SetType& Out, T&& Item) const
 	{
-		if constexpr (Op == EMember::Add)
+		if constexpr (Op == EOp::Add)
 		{
-			Out.Emplace(MoveTemp(Item));
+			Out.Add(MoveTemp(Item));
 		}
 		else
 		{
@@ -664,16 +716,29 @@ struct TSetDeltaBinding : ICustomBinding, FSetDeltaIds
 		}
 	}
 
-	inline static bool Diff(const FSet& A, const FSet& B)
+	inline static bool ContainsValue(const Type& Set, const ElemType& Elem)
 	{
-		if (A.Num() != B.Num())
+		if constexpr (bIsSet)
+		{
+			return Set.Contains(Elem);
+		}
+		else
+		{
+			const auto* Value = Set.Find(Elem.Key);
+			return Value && *Value == Elem.Value;
+		}
+	}
+
+	inline static bool Diff(const Type& As, const Type& Bs)
+	{
+		if (As.Num() != Bs.Num())
 		{
 			return true;
 		}
 
-		for (const T& AKey : A)
+		for (const ElemType& A : As)
 		{
-			if (!B.Contains(AKey))
+			if (!ContainsValue(Bs, A))
 			{
 				return true;
 			}
@@ -683,6 +748,70 @@ struct TSetDeltaBinding : ICustomBinding, FSetDeltaIds
 	}
 };
 
+template <typename T, typename KeyFuncs, typename SetAllocator>
+struct TSetDeltaBinding : TBaseDeltaBinding<TSet<T, KeyFuncs, SetAllocator>, T>
+{
+	using SetType = TSet<T, KeyFuncs, SetAllocator>;
+	using Super = TBaseDeltaBinding<SetType, T>;
+
+	struct FCustomTypename
+	{
+		inline static constexpr std::string_view DeclName = "SetDelta";
+		inline static constexpr std::string_view BindName = Concat<DeclName, ShortTypename<KeyFuncs>, ShortTypename<SetAllocator>>;
+		inline static constexpr std::string_view Namespace;
+		using Parameters = std::tuple<T>;
+	};
+		
+	template<class Ids>	
+	TSetDeltaBinding(TCustomInit<Ids> Init)
+	: Super(Init)
+	{}
+	
+	inline void Save(FMemberBuilder& Dst, const SetType& Src, const SetType* Default, const FSaveContext& Ctx) const
+	{
+		Super::SaveDelta(Dst, Src, Default, Ctx, /* Keys */ Super::Elems);
+	}
+
+	inline void Load(SetType& Dst, FStructLoadView Src, ECustomLoadMethod Method) const
+	{
+		Super::LoadDelta(Dst, Src, Method, /* Keys */ Super::Elems);
+	}
+};
+
+template <typename K, typename V, typename SetAllocator, typename KeyFuncs>
+struct TMapDeltaBinding : TBaseDeltaBinding<TMap<K, V, SetAllocator, KeyFuncs>, K>
+{
+	using MapType = TMap<K, V, SetAllocator, KeyFuncs>;
+	using Super = TBaseDeltaBinding<MapType, K>;
+
+	struct FCustomTypename
+	{
+		inline static constexpr std::string_view DeclName = "MapDelta";
+		inline static constexpr std::string_view BindName = Concat<DeclName, ShortTypename<KeyFuncs>, ShortTypename<SetAllocator>>;
+		inline static constexpr std::string_view Namespace;
+		using Parameters = std::tuple<K, V>;
+	};
+	
+	const TMemberSerializer<K> Keys;
+	
+	template<class Ids>	
+	TMapDeltaBinding(TCustomInit<Ids> Init)
+	: Super(Init)
+	, Keys(Init, {FBaseDeltaBinding::DelId()})
+	{}
+	
+	inline void Save(FMemberBuilder& Dst, const MapType& Src, const MapType* Default, const FSaveContext& Ctx) const
+	{
+		Super::SaveDelta(Dst, Src, Default, Ctx, Keys);
+	}
+
+	inline void Load(MapType& Dst, FStructLoadView Src, ECustomLoadMethod Method) const
+	{
+		Super::LoadDelta(Dst, Src, Method, Keys);
+	}
+};
+
+
 //	template <typename T, typename VariantType>
 //	struct TVariantConstructFromMember
 //	{
@@ -863,10 +992,25 @@ struct TRangeBind<TOptional<T>>
 	using Type = UE::TOptionalBinding<T>;
 };
 
+template<>
+struct TOccupancyOf<FQuat> : FRequireAll {};
+
+template<>
+struct TOccupancyOf<FVector> : FRequireAll {};
+
 template<>
 struct TCustomBind<FTransform>
 {
 	using Type = UE::FTransformBinding;
 };
 
+template<typename K, typename V> requires (!std::is_default_constructible_v<TPair<K,V>>)
+struct TCustomBind<TPair<K, V>>
+{
+	using Type = UE::TPairBinding<K, V>;
+};
+
+template<typename K, typename V>
+struct TOccupancyOf<TPair<K, V>> : FRequireAll {};
+
 } // namespace PlainProps
\ No newline at end of file
diff --git a/Engine/Plugins/Experimental/PlainProps/Source/Tests/PlainPropsTests.cpp b/Engine/Plugins/Experimental/PlainProps/Source/Tests/PlainPropsTests.cpp
--- a/Engine/Plugins/Experimental/PlainProps/Source/Tests/PlainPropsTests.cpp
+++ b/Engine/Plugins/Experimental/PlainProps/Source/Tests/PlainPropsTests.cpp
@@ -716,7 +716,7 @@ TEST_CASE_NAMED(FPlainPropsReadWriteTest, "System::Core::Serialization::PlainPro
 			CHECK(Objects.Num() == 1);
 
 			FReadBatchId Batch = Objects[0].Schema.Batch;
-			auto GetEnumName = [Batch, &Names](FLeafView Leaf) { return Names[ResolveEnumSchema(Batch, Leaf.Enum).Type.Name]; };
+			auto GetEnumName = [Batch, &Names](FLeafView Leaf) { return Names[ResolveEnumSchema(Batch, Leaf.Enum.Get()).Type.Name]; };
 
 			FTestMemberReader It1(Objects[0]);
 			CHECK(It1.GrabLeaf().AsEnum8() == 2);
diff --git a/Engine/Plugins/Experimental/PlainProps/Source/Tests/PlainPropsUeCoreTests.cpp b/Engine/Plugins/Experimental/PlainProps/Source/Tests/PlainPropsUeCoreTests.cpp
--- a/Engine/Plugins/Experimental/PlainProps/Source/Tests/PlainPropsUeCoreTests.cpp
+++ b/Engine/Plugins/Experimental/PlainProps/Source/Tests/PlainPropsUeCoreTests.cpp
@@ -67,11 +67,11 @@ struct TCustomDeltaBind<TSet<T, KeyFuncs, SetAllocator>>
 	using Type = UE::TSetDeltaBinding<T, KeyFuncs, SetAllocator>;
 };
 
-//template <typename K, typename V, typename SetAllocator, typename KeyFuncs>
-//struct TCustomDeltaBind<TMap<K, V, SetAllocator, KeyFuncs>>
-//{
-//	using Type = UE::TMapDeltaBinding<K, V, SetAllocator, KeyFuncs>;
-//};
+template <typename K, typename V, typename SetAllocator, typename KeyFuncs>
+struct TCustomDeltaBind<TMap<K, V, SetAllocator, KeyFuncs>>
+{
+	using Type = UE::TMapDeltaBinding<K, V, SetAllocator, KeyFuncs>;
+};
 
 struct FDeltaRuntime : FDefaultRuntime
 {
@@ -93,48 +93,8 @@ struct TScopedEnumDeclaration
 	~TScopedEnumDeclaration() { Runtime::GetTypes().DropEnum(Id); }
 };
 
-template<class T, EMemberPresence Occupancy = EMemberPresence::AllowSparse, class Runtime = FDefaultRuntime>
-struct TScopedStructDeclaration
-{
-	using Ids = typename Runtime::Ids;
-
-	FStructSchemaId DeclId;
-
-	TScopedStructDeclaration()
-	: DeclId(DeclareNativeStruct<CttiOf<T>, Ids>(Runtime::GetTypes(), Occupancy))
-	{}
-
-	~TScopedStructDeclaration()
-	{
-		Runtime::GetTypes().DropStructRef(DeclId);
-	}
-
-	const FStructDeclaration& Get() const
-	{
-		return Runtime::GetTypes().Get(DeclId);
-	}
-};
-
-template<class T, EMemberPresence Occupancy = EMemberPresence::AllowSparse, class Runtime = FDefaultRuntime>
-struct TScopedStructBinding : TScopedStructDeclaration<T, Occupancy, Runtime>
-{
-	using Super = TScopedStructDeclaration<T, Occupancy, Runtime>;
-	using Ids = typename Runtime::Ids;
-	using Typename = TTypename<T>;
-
-	FStructSchemaId BindId;
-
-	TScopedStructBinding()
-	: BindId(IndexStructBindIdIfNeeded<Ids, Typename>(Super::DeclId))
-	{
-		BindNativeStruct<CttiOf<T>, Runtime>(Runtime::GetSchemas(), BindId, Super::DeclId);
-	}
-
-	~TScopedStructBinding()
-	{
-		Runtime::GetSchemas().DropStruct(BindId);
-	}
-};
+template<class T>
+using TScopedDefaultStructBinding = TScopedStructBinding<T, FDefaultRuntime>;
 
 ////////////////////////////////////////////////////////////////////////////////////////////////
 
@@ -170,9 +130,9 @@ struct FNameBinding : FTestCustomBinding
 		Dst.Add(Declaration.Idx, Idx.AsInteger());
 	}
 
-	virtual void LoadCustom(void* Dst, FStructView Src, ECustomLoadMethod, const FLoadBatch&) const override
+	virtual void LoadCustom(void* Dst, FStructLoadView Src, ECustomLoadMethod) const override
 	{
-		FSetElementId Idx = FSetElementId::FromInteger(FMemberReader(Src).GrabLeaf().AsS32());
+		FSetElementId Idx = FSetElementId::FromInteger(FMemberLoader(Src).GrabLeaf().AsS32());
 		*static_cast<FName*>(Dst) = Names.Get(Idx);
 	}
 
@@ -214,7 +174,7 @@ private:
 FBatchSaver::FBatchSaver(const FCustomBindings& CustomBase)
 : Customs(/* debug */ GNames, &CustomBase)
 {
-	Customs.BindStruct(SavedNames.Declaration.Id, SavedNames.Declaration.Id, SavedNames);
+	Customs.BindStruct(SavedNames.Declaration.Id, SavedNames.Declaration.Id, {}, SavedNames);
 }
 
 template<class T>
@@ -349,7 +309,7 @@ public:
 
 		// Read names and bind custom loader
 		Names.Names.Append(GrabNumAndArray<FName>(It));
-		Customs.BindStruct(Names.Declaration.Id, Names.Declaration.Id, Names);
+		Customs.BindStruct(Names.Declaration.Id, Names.Declaration.Id, {}, Names);
 		CHECK(It.Grab<uint32>() == Magics[5]);
 
 		// Finally create load plans
@@ -385,7 +345,7 @@ public:
 	void LoadInto(T& Out)
 	{
 		FStructView In = Objects[LoadIdx++];
-		LoadStruct(reinterpret_cast<uint8*>(&Out), In, *Plans);
+		LoadStruct(&Out, In.Values, In.Schema.Id, *Plans);
 	}
 
 private:
@@ -517,6 +477,44 @@ struct FNDC
 };
 PP_REFLECT_STRUCT(PlainProps::UE::Test, FNDC, void, X);
 
+struct FNDCBinding : ICustomBinding
+{
+	using Type = FNDC;
+	const FMemberId MemberIds[1];
+
+	template<class Ids>
+	FNDCBinding(TCustomInit<Ids>)
+	: MemberIds{Ids::IndexMember("X")}
+	{}
+
+	void Save(FMemberBuilder& Dst, const FNDC& Src, const FNDC*, const FSaveContext&) const
+	{
+		Dst.Add(MemberIds[0], Src.X);
+	}
+
+	void Load(FNDC& Dst, FStructLoadView Src, ECustomLoadMethod Method) const
+	{
+		int X = FMemberLoader(Src).GrabLeaf().AsS32();
+		if (Method == ECustomLoadMethod::Construct)
+		{
+			new (&Dst) FNDC { X };
+		}
+		else
+		{
+			Dst.X = X;
+		}
+	}
+
+	static bool Diff(FNDC A, FNDC B)
+	{
+		return !(A == B);
+	}
+};
+
+} namespace PlainProps {
+template<> struct TCustomBind<UE::Test::FNDC> {	using Type = UE::Test::FNDCBinding; };
+} namespace PlainProps::UE::Test {
+
 struct FSets
 {
 	TSet<char> Leaves;
@@ -526,27 +524,45 @@ struct FSets
 };
 PP_REFLECT_STRUCT(PlainProps::UE::Test, FSets, void, Leaves, Ranges, Structs, Strings);
 
+
 template<typename T>
-inline bool OrderEq(const TSet<T>& A, const TSet<T>& B)
+bool ContainsElem(const TSet<T>& Set, const T& Elem)
+{
+	return Set.Contains(Elem);
+}
+
+template<typename K, typename V>
+bool ContainsElem(const TMap<K,V>& Map, const TPair<K,V>& Elem)
+{
+	const auto* Value = Map.Find(Elem.Key);
+	return Value && *Value == Elem.Value;
+}
+
+enum ECompare { Order, Content };
+
+template<ECompare Cmp, typename ContainerType>
+inline bool Equals(const ContainerType& A, const ContainerType& B)
 {
 	if (A.Num() != B.Num())
 	{
 		return false;
 	}
 
-	for (typename TSet<T>::TConstIterator AIt(A), BIt(B); AIt; ++AIt, ++BIt)
+	typename ContainerType::TConstIterator ItB(B);
+	for (const auto& ElemA : A)
 	{
-		if (*AIt != *BIt)
+		if (Cmp == ECompare::Order ? ElemA != *ItB : !ContainsElem(B, ElemA))
 		{
 			return false;
 		}
+		++ItB;
 	}
 	return true;
 }
-
-inline bool operator==(const FSets& A, const FSets& B)
+template<ECompare Cmp>
+bool Same(const FSets& A, const FSets& B)
 {
-	return OrderEq(A.Leaves, B.Leaves) && OrderEq(A.Ranges, B.Ranges) && OrderEq(A.Structs, B.Structs) && OrderEq(A.Strings, B.Strings);
+	return Equals<Cmp>(A.Leaves, B.Leaves) && Equals<Cmp>(A.Ranges, B.Ranges) && Equals<Cmp>(A.Structs, B.Structs) && Equals<Cmp>(A.Strings, B.Strings);
 }
 
 struct FMaps
@@ -554,12 +570,14 @@ struct FMaps
 	TMap<bool, bool> Leaves;
 	TMap<int, TArray<char>> Ranges;
 	TMap<FInt, FNDC> Structs;
+	TMap<FString, FString> Strings; // ILeafRangeBinding
 };
-PP_REFLECT_STRUCT(PlainProps::UE::Test, FMaps, void, Leaves, Ranges, Structs);
+PP_REFLECT_STRUCT(PlainProps::UE::Test, FMaps, void, Leaves, Ranges, Structs, Strings);
 
-inline bool operator==(const FMaps& A, const FMaps& B)
+template<ECompare Cmp>
+bool Same(const FMaps& A, const FMaps& B)
 {
-	return LegacyCompareEqual(A.Leaves, B.Leaves) && LegacyCompareEqual(A.Ranges, B.Ranges) && LegacyCompareEqual(A.Structs, B.Structs);
+	return Equals<Cmp>(A.Leaves, B.Leaves) && Equals<Cmp>(A.Ranges, B.Ranges) && Equals<Cmp>(A.Structs, B.Structs) && Equals<Cmp>(A.Strings, B.Strings);
 }
 
 struct FIntAlias
@@ -709,8 +727,6 @@ PP_REFLECT_STRUCT(PlainProps::UE::Test, FDelta, void, A, B, C, D, E);
 //{
 //	using Type = FObject*;
 //
-//	static constexpr EMemberPresence Occupancy = EMemberPresence::RequireAll;
-//
 //	static TConstArrayView<FMemberId> GetMemberIds()
 //	{
 //		static FMemberId Ids[] = {Ids::IndexMember("Type"), Ids::IndexMember("Id")};
@@ -778,7 +794,7 @@ TEST_CASE_NAMED(FPlainPropsUeCoreTest, "System::Core::Serialization::PlainProps:
 {
 	SECTION("Basic")
 	{
-		TScopedStructBinding<FInt> Int;
+		TScopedDefaultStructBinding<FInt> Int;
 		Run([](FBatchSaver& Batch)
 			{
 				Batch.Save(FInt{1234});
@@ -795,7 +811,7 @@ TEST_CASE_NAMED(FPlainPropsUeCoreTest, "System::Core::Serialization::PlainProps:
 		TScopedEnumDeclaration<EFlat2, EEnumMode::Flat> Flat2;
 		TScopedEnumDeclaration<EFlag1, EEnumMode::Flag> Flag1;
 		TScopedEnumDeclaration<EFlag2, EEnumMode::Flag> Flag2;
-		TScopedStructBinding<FEnums> Int;
+		TScopedDefaultStructBinding<FEnums> Int;
 		Run([](FBatchSaver& Batch)
 			{
 				Batch.Save(FEnums{EFlat1::A, EFlat2::A, EFlag1::A, EFlag2::A});
@@ -816,7 +832,7 @@ TEST_CASE_NAMED(FPlainPropsUeCoreTest, "System::Core::Serialization::PlainProps:
 
 	SECTION("TArray")
 	{
-		TScopedStructBinding<FLeafArrays> LeafArrays;
+		TScopedDefaultStructBinding<FLeafArrays> LeafArrays;
 		Run([](FBatchSaver& Batch)
 			{
 				Batch.Save(FLeafArrays{{}, {}});
@@ -836,8 +852,8 @@ TEST_CASE_NAMED(FPlainPropsUeCoreTest, "System::Core::Serialization::PlainProps:
 	SECTION("Nesting")
 	{
 		TScopedEnumDeclaration<EFlat1, EEnumMode::Flat> Flat1;
-		TScopedStructBinding<FLeafArrays> LeafArrays;
-		TScopedStructBinding<FComplexArrays> ComplexArrays;
+		TScopedDefaultStructBinding<FLeafArrays> LeafArrays;
+		TScopedDefaultStructBinding<FComplexArrays> ComplexArrays;
 		Run([](FBatchSaver& Batch)
 			{
 				Batch.Save(FComplexArrays{});
@@ -852,8 +868,8 @@ TEST_CASE_NAMED(FPlainPropsUeCoreTest, "System::Core::Serialization::PlainProps:
 
 	SECTION("TUniquePtr")
 	{
-		TScopedStructBinding<FInt> Int;
-		TScopedStructBinding<FUniquePtrs> UniquePtrs;
+		TScopedDefaultStructBinding<FInt> Int;
+		TScopedDefaultStructBinding<FUniquePtrs> UniquePtrs;
 		Run([](FBatchSaver& Batch)
 			{
 				Batch.Save(FUniquePtrs{});
@@ -868,9 +884,9 @@ TEST_CASE_NAMED(FPlainPropsUeCoreTest, "System::Core::Serialization::PlainProps:
 
 	SECTION("TOptional")
 	{
-		TScopedStructBinding<FNDC> NDC;
-		TScopedStructBinding<FNDCIntrusive> NDCI;
-		TScopedStructBinding<FOpts> Opts;
+		TScopedDefaultStructBinding<FNDC> NDC;
+		TScopedDefaultStructBinding<FNDCIntrusive> NDCI;
+		TScopedDefaultStructBinding<FOpts> Opts;
 		Run([](FBatchSaver& Batch)
 			{
 				Batch.Save(FOpts{});
@@ -887,7 +903,7 @@ TEST_CASE_NAMED(FPlainPropsUeCoreTest, "System::Core::Serialization::PlainProps:
 
 	SECTION("FName")
 	{
-		TScopedStructBinding<FNames> Names;
+		TScopedDefaultStructBinding<FNames> Names;
 		Run([](FBatchSaver& Batch)
 			{
 				Batch.Save(FNames{ FName("A"), {FName("Y"), FName("A")} });
@@ -900,7 +916,7 @@ TEST_CASE_NAMED(FPlainPropsUeCoreTest, "System::Core::Serialization::PlainProps:
 
 	SECTION("FString")
 	{
-		TScopedStructBinding<FStr> Str;
+		TScopedDefaultStructBinding<FStr> Str;
 		Run([](FBatchSaver& Batch)
 			{
 				Batch.Save(FStr{});
@@ -929,12 +945,12 @@ TEST_CASE_NAMED(FPlainPropsUeCoreTest, "System::Core::Serialization::PlainProps:
 
 	SECTION("TypeErasure")
 	{
-		TScopedStructBinding<FSame1> Same1;
-		TScopedStructBinding<FSame2> Same2;	
-		TScopedStructBinding<FIntAlias> IntAlias;
-		TScopedStructBinding<TPair<FString, TArray<char8_t>>> X;
-		TScopedStructBinding<TPair<TArray<char>, TArray<char, TInlineAllocator<8>>>> Y;
-		TScopedStructBinding<FTypeErasure> TypeErasure;
+		TScopedDefaultStructBinding<FSame1> Same1;
+		TScopedDefaultStructBinding<FSame2> Same2;	
+		TScopedDefaultStructBinding<FIntAlias> IntAlias;
+		TScopedDefaultStructBinding<TPair<FString, TArray<char8_t>>> X;
+		TScopedDefaultStructBinding<TPair<TArray<char>, TArray<char, TInlineAllocator<8>>>> Y;
+		TScopedDefaultStructBinding<FTypeErasure> TypeErasure;
 
 		Run([](FBatchSaver& Batch)
 			{
@@ -950,8 +966,8 @@ TEST_CASE_NAMED(FPlainPropsUeCoreTest, "System::Core::Serialization::PlainProps:
 
 	SECTION("TSet")
 	{
-		TScopedStructBinding<FInt> Int;
-		TScopedStructBinding<FSets> Sets;
+		TScopedDefaultStructBinding<FInt> Int;
+		TScopedDefaultStructBinding<FSets> Sets;
 		Run([](FBatchSaver& Batch)
 			{
 				Batch.Save(FSets{{'H','i'}, {{uint8(10)}, {}}, {{123}}});
@@ -969,39 +985,42 @@ TEST_CASE_NAMED(FPlainPropsUeCoreTest, "System::Core::Serialization::PlainProps:
 			}, 
 			[](FBatchLoader& Batch)
 			{
-				CHECK(Batch.Load<FSets>() == FSets{{'H','i'}, {{uint8(10)}, {}}, {{123}}});
-				CHECK(FSets{{'a','b'}} != FSets{{'b','a'}});
-				CHECK(Batch.Load<FSets>() == FSets{{'a','b'}});
-				CHECK(Batch.Load<FSets>() == FSets{{'b','a'}});
-				CHECK(Batch.Load<FSets>() == FSets{{'z','a','?'}});
+				CHECK(Same<ECompare::Order>(Batch.Load<FSets>(), FSets{{'H','i'}, {{uint8(10)}, {}}, {{123}}}));
+				CHECK(!Same<ECompare::Order>(FSets{{'a','b'}}, FSets{{'b','a'}}));
+				CHECK(Same<ECompare::Order>(Batch.Load<FSets>(), FSets{{'a','b'}}));
+				CHECK(Same<ECompare::Order>(Batch.Load<FSets>(), FSets{{'b','a'}}));
+				CHECK(Same<ECompare::Order>(Batch.Load<FSets>(), FSets{{'z','a','?'}}));
 			});
 	}
 
 	SECTION("TMap")
 	{
-		TScopedStructBinding<FInt> Int;
-		TScopedStructBinding<FNDC> NDC;
-		TScopedStructBinding<FMaps> Maps;
-		TScopedStructBinding<TPair<bool, bool>> BoolBoolPair;
-		TScopedStructBinding<TPair<int, TArray<char>>> IntStringPair;
-		TScopedStructBinding<TPair<FInt, FNDC>> IntNDCPair;
+		TScopedDefaultStructBinding<FInt> Int;
+		TScopedDefaultStructBinding<FNDC> NDC;
+		TScopedDefaultStructBinding<FMaps> Maps;
+		// Todo: Recursive auto-bind mechanism?
+		TScopedDefaultStructBinding<TPair<bool, bool>> BoolBoolPair;
+		TScopedDefaultStructBinding<TPair<int, TArray<char>>> IntStringPair;
+		TScopedDefaultStructBinding<TPair<FInt, FNDC>> IntNDCPair;
+		TScopedDefaultStructBinding<TPair<FString, FString>> StrStrPair;
 		
 		Run([](FBatchSaver& Batch)
 			{
+				TPair<FString, FString> abc = {"a", "bc"};
 				Batch.Save(FMaps{});
-				Batch.Save(FMaps{{{true, true}, {false, false}}, {{5, {'h', 'i'}}}, {{{7}, FNDC{8}}}});
+				Batch.Save(FMaps{{{true, true}, {false, false}}, {{5, {'h', 'i'}}}, {{{7}, FNDC{8}}}, {{"a", "bc"}}});
 			}, 
 			[](FBatchLoader& Batch)
 			{
-				CHECK(Batch.Load<FMaps>() == FMaps{});
-				CHECK(Batch.Load<FMaps>() == FMaps{{{true, true}, {false, false}}, {{5, {'h', 'i'}}}, {{{7}, FNDC{8}}}});
+				CHECK(Same<ECompare::Order>(Batch.Load<FMaps>(), FMaps{}));
+				CHECK(Same<ECompare::Order>(Batch.Load<FMaps>(), FMaps{{{true, true}, {false, false}}, {{5, {'h', 'i'}}}, {{{7}, FNDC{8}}}, {{"a", "bc"}}}));
 			});
 	}
 
 	SECTION("Delta")
 	{
-		TScopedStructBinding<FInt> Int;
-		TScopedStructBinding<FDelta> Delta;
+		TScopedDefaultStructBinding<FInt> Int;
+		TScopedDefaultStructBinding<FDelta> Delta;
 		Run([](FBatchSaver& Batch)
 			{
 				FDelta Zero = {false, 0, {}, {}, {}};
@@ -1030,8 +1049,8 @@ TEST_CASE_NAMED(FPlainPropsUeCoreTest, "System::Core::Serialization::PlainProps:
 
 	SECTION("TSetDelta")
 	{
-		TScopedStructBinding<FInt> Int;
-		TScopedStructBinding<FSets, EMemberPresence::AllowSparse, FDeltaRuntime> Sets;
+		TScopedDefaultStructBinding<FInt> Int;
+		TScopedStructBinding<FSets, FDeltaRuntime> Sets;
 		Run([](FBatchSaver& Batch)
 			{
 				Batch.Save(FSets{});
@@ -1045,22 +1064,54 @@ TEST_CASE_NAMED(FPlainPropsUeCoreTest, "System::Core::Serialization::PlainProps:
 			}, 
 			[](FBatchLoader& Batch)
 			{
-				CHECK(Batch.Load<FSets>() == FSets{});
-				CHECK(Batch.Load<FSets>() == FSets{{'l'}, {{1}}, {{2}}, {"s"}});
+				CHECK(Same<ECompare::Content>(Batch.Load<FSets>(), FSets{}));
+				CHECK(Same<ECompare::Content>(Batch.Load<FSets>(), FSets{{'l'}, {{1}}, {{2}}, {"s"}}));
 				
 				const FSets Default = {{'a'}, {{1}}, {{1}}, {"a"}};
-				CHECK(Batch.Load<FSets>() == Default);
-				CHECK(Batch.Load(Default) == FSets{}); // Wipe
-				CHECK(Batch.Load(Default) == FSets{{'a'}, {{0,1,2}}, {{2}}, {}}); // Mixed changes
+				CHECK(Same<ECompare::Content>(Batch.Load<FSets>(), Default));
+				CHECK(Same<ECompare::Content>(Batch.Load(Default), FSets{})); // Wipe
+				CHECK(Same<ECompare::Content>(Batch.Load(Default), FSets{{'a'}, {{0,1,2}}, {{2}}, {}})); // Mixed changes
+			},
+			GDeltaCustoms);
+	}
+
+	SECTION("TMapDelta")
+	{
+		TScopedDefaultStructBinding<FInt> Int;
+		TScopedDefaultStructBinding<FNDC> NDC;
+		// Todo: Recursive auto-bind mechanism?
+		TScopedDefaultStructBinding<TPair<bool, bool>> BoolBoolPair;
+		TScopedDefaultStructBinding<TPair<int, TArray<char>>> IntStringPair;
+		TScopedDefaultStructBinding<TPair<FInt, FNDC>> IntNDCPair;
+		TScopedDefaultStructBinding<TPair<FString, FString>> StrStrPair;
+		TScopedStructBinding<FMaps, FDeltaRuntime> Maps;
+
+		Run([](FBatchSaver& Batch)
+			{
+				const FMaps Default = {{{true,true}}, {{1, {'a'}}},	{{{2}, FNDC{3}}},	{{"hi", "lo"}}};
+				const FMaps Changes = {{},	{{0,{'a'}}, {2,{'a'}}},	{{{2}, FNDC{4}}},	{{"hi", "hi"}}};
+				CHECK(!Batch.SaveDelta(FMaps{}, FMaps{}));
+				CHECK(Batch.SaveDelta(Default, FMaps{}));
+				CHECK(Batch.SaveDelta(FMaps{}, Default)); // Wipe defaults
+				CHECK(!Batch.SaveDelta(Default, Default));
+				CHECK(Batch.SaveDelta(Changes, Default));
+			}, 
+			[](FBatchLoader& Batch)
+			{
+				const FMaps Default = {{{true,true}}, {{1, {'a'}}},	{{{2}, FNDC{3}}},	{{"hi", "lo"}}};
+				const FMaps Changes = {{},	{{0,{'a'}}, {2,{'a'}}},	{{{2}, FNDC{4}}},	{{"hi", "hi"}}};
+				CHECK(Same<ECompare::Content>(Batch.Load<FMaps>(), Default));
+				CHECK(Same<ECompare::Content>(Batch.Load<FMaps>(), FMaps{})); // Wipe defaults
+				CHECK(Same<ECompare::Content>(Batch.Load(Default), Changes));
 			},
 			GDeltaCustoms);
 	}
 
 	SECTION("Transform")
 	{
-		TScopedStructBinding<FVector, EMemberPresence::RequireAll> Vector;
-		TScopedStructBinding<FQuat, EMemberPresence::RequireAll> Quat;
-		BindCustomStructOnce<FTransformBinding, FDefaultRuntime>();
+		TScopedDefaultStructBinding<FVector> Vector;
+		TScopedDefaultStructBinding<FQuat> Quat;
+		TScopedDefaultStructBinding<FTransform> Transform;
 
 		Run([](FBatchSaver& Batch)
 			{
