Clang fix

#rnx

diff --git a/Engine/Plugins/Experimental/PlainProps/Source/Public/PlainPropsLoad.h b/Engine/Plugins/Experimental/PlainProps/Source/Public/PlainPropsLoad.h
--- a/Engine/Plugins/Experimental/PlainProps/Source/Public/PlainPropsLoad.h
+++ b/Engine/Plugins/Experimental/PlainProps/Source/Public/PlainPropsLoad.h
@@ -27,7 +27,8 @@ template<class Runtime>
 	return CreateLoadPlans(ReadId, Runtime::GetTypes(), Runtime::GetCustoms(), Runtime::GetSchemas(), RuntimeIds);
 }
 
-PLAINPROPS_API [[nodiscard]] FLoadBatchPtr	CreateLoadPlans(FReadBatchId ReadId, const FDeclarations& Declarations, const FCustomBindings& Customs, const FSchemaBindings& Schemas, TConstArrayView<FStructSchemaId> RuntimeIds);
+
+[[nodiscard]] PLAINPROPS_API FLoadBatchPtr	CreateLoadPlans(FReadBatchId ReadId, const FDeclarations& Declarations, const FCustomBindings& Customs, const FSchemaBindings& Schemas, TConstArrayView<FStructSchemaId> RuntimeIds);
 PLAINPROPS_API void							LoadStruct(void* Dst, FByteReader Src, FStructSchemaId Id, const FLoadBatch& Batch);
 PLAINPROPS_API void							ConstructAndLoadStruct(void* Dst, FByteReader Src, FStructSchemaId LoadId, const FLoadBatch& Batch);
 
