PlainProps: Minor read & write bug fixes related to inheritance.
Fix array view OOBs if calling FStructSchema::GetMemberNames/EditMemberNames for structs with a super member.
Fix broken GetInheritance() write logic that was comparing the remapped declared super schema with the non remapped member super schema.
Remove the convenience function FStructSchemaHandle::ResolveSuper() that was unsafe as well as likely to do an unnecessary re-resolve of the struct itself.
Expand the unit tests for "Inheritance" and "SparseInheritance" to get coverage of the two fixed issues.

#rnx

diff --git a/Engine/Plugins/Experimental/PlainProps/Source/Private/PlainPropsInternalFormat.h b/Engine/Plugins/Experimental/PlainProps/Source/Private/PlainPropsInternalFormat.h
--- a/Engine/Plugins/Experimental/PlainProps/Source/Private/PlainPropsInternalFormat.h
+++ b/Engine/Plugins/Experimental/PlainProps/Source/Private/PlainPropsInternalFormat.h
@@ -64,13 +64,14 @@ struct FStructSchema
 	uint8 IsDense : 1;
 	FMemberType Footer[0];
 	
+	inline uint32					NumNames() const			{ checkSlow(NumMembers > 0 || !UsesSuper(Inheritance)); return NumMembers - UsesSuper(Inheritance); }
 	TConstArrayView<FMemberType>	GetMemberTypes() const		{ return MakeArrayView(GetMemberTypes(Footer), NumMembers); }
 	TConstArrayView<FMemberType>	GetRangeTypes() const		{ return MakeArrayView(GetRangeTypes(Footer, NumMembers), NumRangeTypes); }
-	TConstArrayView<FMemberId>		GetMemberNames() const		{ return MakeArrayView(GetMemberNames(Footer, NumMembers, NumRangeTypes), NumMembers); }
-	TArrayView<FMemberId>			EditMemberNames() 			{ return MakeArrayView(const_cast<FMemberId*>(GetMemberNames(Footer, NumMembers, NumRangeTypes)), NumMembers); }
-	const FSchemaId*				GetInnerSchemas() const		{ return GetInnerSchemas(Footer, NumMembers, NumRangeTypes, NumMembers - UsesSuper(Inheritance)); }
+	TConstArrayView<FMemberId>		GetMemberNames() const		{ return MakeArrayView(GetMemberNames(Footer, NumMembers, NumRangeTypes), NumNames()); }
+	TArrayView<FMemberId>			EditMemberNames() 			{ return MakeArrayView(const_cast<FMemberId*>(GetMemberNames(Footer, NumMembers, NumRangeTypes)), NumNames()); }
+	const FSchemaId*				GetInnerSchemas() const		{ return GetInnerSchemas(Footer, NumMembers, NumRangeTypes, NumNames()); }
 	FOptionalStructSchemaId			GetSuperSchema() const		{ return Inheritance != ESuper::No ? ToOptional(static_cast<FStructSchemaId>(*GetInnerSchemas())) : NoId; }
-	
+
 	static const FMemberType*		GetMemberTypes(const FMemberType* Footer)
 	{
 		return Footer;
@@ -157,4 +158,4 @@ inline constexpr uint64 GetLeafRangeSize(uint64 Num, FUnpackedLeafType Leaf)
 
 //////////////////////////////////////////////////////////////////////////
 
-} // namespace PlainProps
\ No newline at end of file
+} // namespace PlainProps
diff --git a/Engine/Plugins/Experimental/PlainProps/Source/Private/PlainPropsRead.cpp b/Engine/Plugins/Experimental/PlainProps/Source/Private/PlainPropsRead.cpp
--- a/Engine/Plugins/Experimental/PlainProps/Source/Private/PlainPropsRead.cpp
+++ b/Engine/Plugins/Experimental/PlainProps/Source/Private/PlainPropsRead.cpp
@@ -165,13 +165,6 @@ FNestedRangeView FRangeView::AsRanges() const
 
 //////////////////////////////////////////////////////////////////////////
 
-const FStructSchema& FStructSchemaHandle::ResolveSuper() const
-{
-	return ResolveStructSchema(Batch, Resolve().GetSuperSchema().Get());
-}
-
-//////////////////////////////////////////////////////////////////////////
-
 FRangeView FNestedRangeIterator::operator*() const 
 {
 	FByteReader PeekBytes = ByteIt;
@@ -441,4 +434,4 @@ FFlatMemberReader::FFlatMemberReader(FStructView Struct)
 	It = &Lineage.Last();
 }
 
-} // namespace PlainProps
\ No newline at end of file
+} // namespace PlainProps
diff --git a/Engine/Plugins/Experimental/PlainProps/Source/Private/PlainPropsWrite.cpp b/Engine/Plugins/Experimental/PlainProps/Source/Private/PlainPropsWrite.cpp
--- a/Engine/Plugins/Experimental/PlainProps/Source/Private/PlainPropsWrite.cpp
+++ b/Engine/Plugins/Experimental/PlainProps/Source/Private/PlainPropsWrite.cpp
@@ -295,14 +295,13 @@ static FOptionalSchemaId GetStaticInnerSchema(const FMemberSchema& Schema, const
 		check(IsStructOrEnum(InnermostType));
 		if (InnermostType.IsLeaf())
 		{
-			return ToOptional(static_cast<FSchemaId>(NewIds.RemapEnum(InnerSchema.Get())));
+			return { NewIds.RemapEnum(InnerSchema.Get()) };
 		}
 		else if (!InnermostType.AsStruct().IsDynamic)
 		{
-			return ToOptional(static_cast<FSchemaId>(NewIds.RemapStruct(InnerSchema.Get())));
+			return { NewIds.RemapStruct(InnerSchema.Get()) };
 		}
 	}
-
 	return {};
 }
 
@@ -335,7 +334,7 @@ static TArray<IdType> RemapIds(const FWriteIds& NewIds, const TArray<IdType>& Na
 	return Out;
 }
 
-static ESuper GetInheritance(FOptionalSchemaId Super, TConstArrayView<const FMemberSchema*> Members)
+static ESuper GetInheritance(FOptionalStructSchemaId Super, TConstArrayView<const FMemberSchema*> Members)
 {
 	if (Super)
 	{
@@ -343,8 +342,14 @@ static ESuper GetInheritance(FOptionalSchemaId Super, TConstArrayView<const FMem
 		{
 			return ESuper::Unused;
 		}
-		
-		return Members[0]->Type.AsStruct().IsDynamic || Members[0]->InnerSchema != Super ? ESuper::Used : ESuper::Reused;
+
+		if (Super != static_cast<FOptionalStructSchemaId>(Members[0]->InnerSchema))
+		{
+			return ESuper::Used;
+		}
+
+		check(!Members[0]->Type.AsStruct().IsDynamic);
+		return ESuper::Reused;
 	}
 
 	return ESuper::No;
@@ -352,9 +357,7 @@ static ESuper GetInheritance(FOptionalSchemaId Super, TConstArrayView<const FMem
 
 static void WriteSchema(TArray64<uint8>& Out, const FBuiltStructSchema& Struct, const FWriteIds& NewIds)
 {
-	FOptionalSchemaId NewSuperId = Struct.Super ? ToOptionalSchema(NewIds.RemapStruct(Struct.Super.Get())) : NoId;
-	ESuper Inheritance = GetInheritance(NewSuperId, Struct.MemberSchemas);
-
+	ESuper Inheritance = GetInheritance(Struct.Super, Struct.MemberSchemas);
 	TArray<FMemberType> MemberTypes = GetMemberTypes(Struct);
 	TArray<FMemberType> InnerRangeTypes = GetInnerRangeTypes(Struct);
 	TArray<FMemberId> MemberNames = NewIds.HasStableNames() ? RemapIds(NewIds, Struct.MemberNames) : Struct.MemberNames;
diff --git a/Engine/Plugins/Experimental/PlainProps/Source/Public/PlainPropsRead.h b/Engine/Plugins/Experimental/PlainProps/Source/Public/PlainPropsRead.h
--- a/Engine/Plugins/Experimental/PlainProps/Source/Public/PlainPropsRead.h
+++ b/Engine/Plugins/Experimental/PlainProps/Source/Public/PlainPropsRead.h
@@ -160,7 +160,6 @@ struct FStructSchemaHandle
 	FReadBatchId			Batch;
 
 	const FStructSchema&	Resolve() const { return ResolveStructSchema(Batch, Id); }
-	const FStructSchema&	ResolveSuper() const;
 };
 
 struct FStructView
@@ -669,4 +668,4 @@ private:
 	}
 };
 
-} // namespace PlainProps
\ No newline at end of file
+} // namespace PlainProps
diff --git a/Engine/Plugins/Experimental/PlainProps/Source/Tests/PlainPropsTests.cpp b/Engine/Plugins/Experimental/PlainProps/Source/Tests/PlainPropsTests.cpp
--- a/Engine/Plugins/Experimental/PlainProps/Source/Tests/PlainPropsTests.cpp
+++ b/Engine/Plugins/Experimental/PlainProps/Source/Tests/PlainPropsTests.cpp
@@ -1128,9 +1128,12 @@ TEST_CASE_NAMED(FPlainPropsReadWriteTest, "System::Core::Serialization::PlainPro
 		{
 			CHECK(Objects.Num() == 1);
 
-			FTestMemberReader TopIt(Objects[0]);
-			FTestMemberReader MidIt(TopIt.GrabStruct());
-			FTestMemberReader LowIt(MidIt.GrabStruct());
+			FStructView TopView = Objects[0];
+			FTestMemberReader TopIt(TopView);
+			FStructView MidView = TopIt.GrabStruct();
+			FTestMemberReader MidIt(MidView);
+			FStructView LowView = MidIt.GrabStruct();
+			FTestMemberReader LowIt(LowView);
 			CHECK(LowIt.GrabLeaf().AsS32() == 123);
 			CHECK(MidIt.GrabLeaf().AsS32() == 456);
 			CHECK(FTestMemberReader(MidIt.GrabStruct()).GrabLeaf().AsS32() == 1000);
@@ -1141,6 +1144,17 @@ TEST_CASE_NAMED(FPlainPropsReadWriteTest, "System::Core::Serialization::PlainPro
 			CHECK(Names[TopIt.PeekName()] == "TMids");
 			CHECK(MakeArray<FTestMemberReader>(TopIt.GrabRange().AsStructs())[0].GrabLeaf().AsS32() == 3000);
 
+			const FStructSchema& TopSchema = TopView.Schema.Resolve();
+			const FStructSchema& MidSchema = MidView.Schema.Resolve();
+			const FStructSchema& LowSchema = LowView.Schema.Resolve();
+			CHECK(TopSchema.Inheritance == ESuper::Reused);
+			CHECK(MidSchema.Inheritance == ESuper::Reused);
+			CHECK(LowSchema.Inheritance == ESuper::No);
+			CHECK(TopSchema.NumMembers == 4);
+			CHECK(TopSchema.NumNames() == 3);
+			CHECK(MidSchema.NumMembers == MidSchema.NumNames() + 1);
+			CHECK(LowSchema.NumMembers == LowSchema.NumNames());
+
 			FFlatMemberReader FlatIt(Objects[0]);
 			CHECK(Names[FlatIt.PeekOwner().Name] == "Low");
 			CHECK(FlatIt.GrabLeaf().AsS32() == 123);
@@ -1163,11 +1177,11 @@ TEST_CASE_NAMED(FPlainPropsReadWriteTest, "System::Core::Serialization::PlainPro
 		TestSerialize([](FTestBatchBuilder& Batch, FScratchAllocator& Scratch)
 		{																										// Usage by A B C 
 			FStructSchemaId B0 = Batch.DeclareStruct("Test", "B0", {"0"}, EMemberPresence::AllowSparse);				 // - - -
-			FStructSchemaId B1 = Batch.DeclareStruct("Test", "B1", {"1"}, EMemberPresence::AllowSparse, ToOptional(B0)); // 1 - 0
+			FStructSchemaId B1 = Batch.DeclareStruct("Test", "B1", {"1"}, EMemberPresence::AllowSparse, ToOptional(B0)); // 1 1 1
 			FStructSchemaId B2 = Batch.DeclareStruct("Test", "B2", {"2"}, EMemberPresence::AllowSparse, ToOptional(B1)); // - 1 1
 			FStructSchemaId B3 = Batch.DeclareStruct("Test", "B3", {"3"}, EMemberPresence::AllowSparse, ToOptional(B2)); // - - 0
-			FStructSchemaId B4 = Batch.DeclareStruct("Test", "B4", {"4"}, EMemberPresence::AllowSparse, ToOptional(B3)); // 1 1 0
-			FStructSchemaId B5 = Batch.DeclareStruct("Test", "B5", {"5"}, EMemberPresence::AllowSparse, ToOptional(B4)); // 1 0 0
+			FStructSchemaId B4 = Batch.DeclareStruct("Test", "B4", {"4"}, EMemberPresence::AllowSparse, ToOptional(B3)); // 1 1 1
+			FStructSchemaId B5 = Batch.DeclareStruct("Test", "B5", {"5"}, EMemberPresence::AllowSparse, ToOptional(B4)); // 1 1 0
 			FStructSchemaId B6 = Batch.DeclareStruct("Test", "B6", {"6"}, EMemberPresence::AllowSparse, ToOptional(B5)); // 0 - 0
 			FStructSchemaId C5 = Batch.DeclareStruct("Test", "C5", {"5"}, EMemberPresence::AllowSparse, ToOptional(B4)); // - - -
 			
@@ -1180,13 +1194,16 @@ TEST_CASE_NAMED(FPlainPropsReadWriteTest, "System::Core::Serialization::PlainPro
 			A.BuildSuperStruct(Scratch, Batch.Get(B5), Batch);
 
 			FMemberBuilder B;
+			B.Add(Batch.NameMember("1"), 10);
+			B.BuildSuperStruct(Scratch, Batch.Get(B1), Batch);
 			B.Add(Batch.NameMember("2"), 20);
 			B.BuildSuperStruct(Scratch, Batch.Get(B2), Batch);
 			B.Add(Batch.NameMember("4"), 40);
 			B.BuildSuperStruct(Scratch, Batch.Get(B4), Batch);
 
 			FMemberBuilder C;
-			C.BuildSuperStruct(Scratch, Batch.Get(B1), Batch); // Empty -> noop
+			C.Add(Batch.NameMember("1"), 100);
+			C.BuildSuperStruct(Scratch, Batch.Get(B1), Batch);
 			C.Add(Batch.NameMember("2"), 200);
 			C.BuildSuperStruct(Scratch, Batch.Get(B2), Batch);
 			C.BuildSuperStruct(Scratch, Batch.Get(B3), Batch); // Empty -> noop
@@ -1210,14 +1227,27 @@ TEST_CASE_NAMED(FPlainPropsReadWriteTest, "System::Core::Serialization::PlainPro
 			CHECK(Names[Schema2.Type.Name] == "B6");
 			CHECK(Schema0.GetSuperSchema() == Objects[1].Schema.Id);
 			CHECK(Schema2.GetSuperSchema() == Objects[1].Schema.Id);
-			CHECK(&Objects[0].Schema.ResolveSuper() == &Schema1);
-			
+
 			CHECK(Names[FMemberReader(Objects[0]).GrabStruct().Schema.Resolve().Type.Name]  == "B5");
 			CHECK(Names[FMemberReader(Objects[1]).GrabStruct().Schema.Resolve().Type.Name]  == "B4");
 			CHECK(Names[FMemberReader(Objects[2]).GrabStruct().Schema.Resolve().Type.Name]  == "B4");
 
-			FTestMemberReader B1(Objects[1]);
-			FStructView Super = B1.GrabStruct();
+			FReadBatchId BatchId = Objects[0].Schema.Batch;
+			const FStructSchema& B6 = Objects[0].Schema.Resolve();
+			const FStructSchema& B5 = ResolveStructSchema(BatchId, B6.GetSuperSchema().Get());
+			const FStructSchema& B4 = ResolveStructSchema(BatchId, B5.GetSuperSchema().Get());
+			const FStructSchema& B3 = ResolveStructSchema(BatchId, B4.GetSuperSchema().Get());
+			const FStructSchema& B2 = ResolveStructSchema(BatchId, B3.GetSuperSchema().Get());
+			const FStructSchema& B1 = ResolveStructSchema(BatchId, B2.GetSuperSchema().Get());
+			const FStructSchema& B0 = ResolveStructSchema(BatchId, B1.GetSuperSchema().Get());
+			// Super usage									A	B	C 	Decl
+			CHECK(B0.Inheritance == ESuper::No);		//	-	-	-	-
+			CHECK(B1.Inheritance == ESuper::Unused);	//	0	0	0	B0
+			CHECK(B3.Inheritance == ESuper::Unused);	//	-	-	0	B2
+			CHECK(B2.Inheritance == ESuper::Reused);	//	-	B1	B1	B1
+			CHECK(B4.Inheritance == ESuper::Used);		//	B1	B2	B2	B3
+			CHECK(B5.Inheritance == ESuper::Reused);	//	B4	B4	0	B4
+			CHECK(B6.Inheritance == ESuper::Used);		//	B5	-	B4	B5
 
 			FFlatMemberReader A(Objects[0]);
 			FFlatMemberReader B(Objects[1]);
@@ -1225,8 +1255,10 @@ TEST_CASE_NAMED(FPlainPropsReadWriteTest, "System::Core::Serialization::PlainPro
 			CHECK(A.GrabLeaf().AsS32() == 1);
 			CHECK(A.GrabLeaf().AsS32() == 4);
 			CHECK(A.GrabLeaf().AsS32() == 5);
+			CHECK(B.GrabLeaf().AsS32() == 10);
 			CHECK(B.GrabLeaf().AsS32() == 20);
 			CHECK(B.GrabLeaf().AsS32() == 40);
+			CHECK(C.GrabLeaf().AsS32() == 100);
 			CHECK(C.GrabLeaf().AsS32() == 200);
 			CHECK(C.GrabLeaf().AsS32() == 400);
 			CHECK(!A.HasMore());
@@ -1358,4 +1390,4 @@ TEST_CASE_NAMED(FPlainPropsLoadSaveTest, "System::Core::Serialization::PlainProp
 }
 
 } // namespace PlainProps
-#endif // WITH_TESTS
\ No newline at end of file
+#endif // WITH_TESTS
