PlainProps: Remove the unused ToOptionalSchema() pair of functions, explicitly using the FOptionalSchemaId constructor is cleaner.
Fixup one single call site in BindInnermostType() that did not even use them.

#rnx

diff --git a/Engine/Plugins/Experimental/PlainProps/Source/Public/PlainPropsBind.h b/Engine/Plugins/Experimental/PlainProps/Source/Public/PlainPropsBind.h
--- a/Engine/Plugins/Experimental/PlainProps/Source/Public/PlainPropsBind.h
+++ b/Engine/Plugins/Experimental/PlainProps/Source/Public/PlainPropsBind.h
@@ -917,7 +917,7 @@ FMemberBindType BindInnermostType(FOptionalSchemaId& OutSchema)
 template<Enumeration Enum, class Ids>
 FMemberBindType BindInnermostType(FOptionalSchemaId& OutSchema)
 {
-	OutSchema = ToOptional(static_cast<FSchemaId>(GetEnumId<Ids, Enum>()));
+	OutSchema = FOptionalSchemaId(GetEnumId<Ids, Enum>());
 	return FMemberBindType(ReflectEnum<Enum>);
 }
 
@@ -1202,4 +1202,4 @@ struct FIdTranslator : FIdTranslatorBase
 FSchemaBatch*			CreateTranslatedSchemas(const FSchemaBatch& Schemas, FIdBinding NewIds);
 void					DestroyTranslatedSchemas(const FSchemaBatch* Schemas);
 
-} // namespace PlainProps
\ No newline at end of file
+} // namespace PlainProps
diff --git a/Engine/Plugins/Experimental/PlainProps/Source/Public/PlainPropsTypes.h b/Engine/Plugins/Experimental/PlainProps/Source/Public/PlainPropsTypes.h
--- a/Engine/Plugins/Experimental/PlainProps/Source/Public/PlainPropsTypes.h
+++ b/Engine/Plugins/Experimental/PlainProps/Source/Public/PlainPropsTypes.h
@@ -275,8 +275,6 @@ using FOptionalConcreteTypenameId = TOptionalId<FConcreteTypenameId>;
 
 template<class IdType>
 inline constexpr TOptionalId<IdType> ToOptional(IdType Id) { return Id; }
-inline constexpr FOptionalSchemaId ToOptionalSchema(FEnumSchemaId Id) { return static_cast<FSchemaId>(Id); }
-inline constexpr FOptionalSchemaId ToOptionalSchema(FStructSchemaId Id) { return static_cast<FSchemaId>(Id); }
 
 //////////////////////////////////////////////////////////////////////////
 
@@ -436,4 +434,4 @@ public:
 	PLAINPROPS_API FString				Print(FStructSchemaId Name) const;
 };
 
-} // namespace PlainProps
\ No newline at end of file
+} // namespace PlainProps
