PlainProps: Expand the Enum unit test to cover the built schema and the the recently fixed bug in WriteEnumConstantsAs().

#rnx

diff --git a/Engine/Plugins/Experimental/PlainProps/Source/Tests/PlainPropsTests.cpp b/Engine/Plugins/Experimental/PlainProps/Source/Tests/PlainPropsTests.cpp
--- a/Engine/Plugins/Experimental/PlainProps/Source/Tests/PlainPropsTests.cpp
+++ b/Engine/Plugins/Experimental/PlainProps/Source/Tests/PlainPropsTests.cpp
@@ -1,6 +1,7 @@
 // Copyright Epic Games, Inc. All Rights Reserved.
 
 #if WITH_TESTS
+#include "Algo/Compare.h"
 #include "Tests/TestHarnessAdapter.h"
 #include "Containers/AnsiString.h"
 #include "Containers/Map.h"
@@ -23,7 +24,14 @@ static bool operator==(FParametricTypeView A, FParametricTypeView B)
 {
 	return A.Name == B.Name && A.NumParameters == B.NumParameters && !FMemory::Memcmp(A.Parameters, B.Parameters, A.NumParameters * sizeof(FTypeId));
 }
-	
+
+template<typename IntType>
+static TConstArrayView<IntType> GetEnumConstants(const FEnumSchema& Schema)
+{
+	CHECK(Schema.Width == ReflectLeaf<IntType>.Width);
+	return MakeConstArrayView<IntType>(reinterpret_cast<const IntType*>(Schema.Footer + Schema.Num), Schema.ExplicitConstants ? Schema.Num : 0);
+}
+
 TEST_CASE_NAMED(FPlainPropsIndexTest, "System::Core::Serialization::PlainProps::Index", "[Core][PlainProps][SmokeFilter]")
 {
 	SECTION("NestedScope")
@@ -717,6 +725,10 @@ TEST_CASE_NAMED(FPlainPropsReadWriteTest, "System::Core::Serialization::PlainPro
 
 			FReadBatchId Batch = Objects[0].Schema.Batch;
 			auto GetEnumName = [Batch, &Names](FLeafView Leaf) { return Names[ResolveEnumSchema(Batch, Leaf.Enum.Get()).Type.Name]; };
+			auto EqualEnumNames = [&Names](TConstArrayView<FNameId> Ids, TConstArrayView<const char*> Strings)
+			{
+				return Algo::Compare(Ids, Strings, [&Names](auto& X, auto& Y) { return Names[X] == Y; });
+			};
 
 			FTestMemberReader It1(Objects[0]);
 			CHECK(It1.GrabLeaf().AsEnum8() == 2);
@@ -747,6 +759,53 @@ TEST_CASE_NAMED(FPlainPropsReadWriteTest, "System::Core::Serialization::PlainPro
 			CHECK(GetEnumName(It2.GrabLeaf()) == "FlatLimit32");
 			CHECK(GetEnumName(It2.GrabLeaf()) == "FlatLimit64");
 			CHECK(GetEnumName(It2.GrabLeaf()) == "FlagLimit64");			
+
+			FTestMemberReader It3(Objects[0]);
+			const FEnumSchema& FlatDense8 = ResolveEnumSchema(Batch, It3.GrabLeaf().Enum.Get());
+			(void)It3.GrabLeaf();
+			const FEnumSchema& FlagDense8 = ResolveEnumSchema(Batch, It3.GrabLeaf().Enum.Get());
+			(void)It3.GrabLeaf();
+			(void)It3.GrabLeaf();
+			(void)It3.GrabLeaf();
+			const FEnumSchema& FlatSparse8 = ResolveEnumSchema(Batch, It3.GrabLeaf().Enum.Get());
+			const FEnumSchema& FlagSparse8 = ResolveEnumSchema(Batch, It3.GrabLeaf().Enum.Get());
+			const FEnumSchema& FlatLimit8  = ResolveEnumSchema(Batch, It3.GrabLeaf().Enum.Get());
+			const FEnumSchema& FlatLimit16 = ResolveEnumSchema(Batch, It3.GrabLeaf().Enum.Get());
+			const FEnumSchema& FlatLimit32 = ResolveEnumSchema(Batch, It3.GrabLeaf().Enum.Get());
+			const FEnumSchema& FlatLimit64 = ResolveEnumSchema(Batch, It3.GrabLeaf().Enum.Get());
+			const FEnumSchema& FlagLimit64 = ResolveEnumSchema(Batch, It3.GrabLeaf().Enum.Get());
+
+			CHECK(FlatDense8.ExplicitConstants);
+			CHECK(!FlagDense8.ExplicitConstants);
+			CHECK(FlatSparse8.ExplicitConstants);
+			CHECK(FlagSparse8.ExplicitConstants);
+			CHECK(FlatLimit8.ExplicitConstants);
+			CHECK(FlatLimit16.ExplicitConstants);
+			CHECK(FlatLimit32.ExplicitConstants);
+			CHECK(FlatLimit64.ExplicitConstants);
+			CHECK(FlagLimit64.ExplicitConstants);
+
+			CHECK(EqualEnumNames(MakeConstArrayView(FlatDense8.Footer, FlatDense8.Num), MakeConstArrayView({"A", "C"})));
+			CHECK(EqualEnumNames(MakeConstArrayView(FlagDense8.Footer, FlagDense8.Num), MakeConstArrayView({"A", "B", "C"})));
+			CHECK(EqualEnumNames(MakeConstArrayView(FlatSparse8.Footer, FlatSparse8.Num), MakeConstArrayView({"C"})));
+			CHECK(EqualEnumNames(MakeConstArrayView(FlagSparse8.Footer, FlagSparse8.Num), MakeConstArrayView({"A", "C"})));
+			CHECK(Names[FlatLimit8.Footer[0]] == "Max");
+			CHECK(Names[FlatLimit16.Footer[0]] == "Max");
+			CHECK(Names[FlatLimit32.Footer[0]] == "Max");
+			CHECK(Names[FlatLimit64.Footer[0]] == "Max");
+			CHECK(Names[FlagLimit64.Footer[0]] == "One");
+			CHECK(Names[FlagLimit64.Footer[1]] == "Max");
+
+			CHECK(EqualItems(GetEnumConstants<uint8>(FlatDense8), MakeConstArrayView({0, 2})));
+			CHECK(EqualItems(GetEnumConstants<uint8>(FlagDense8), TConstArrayView<uint8>()));
+			CHECK(EqualItems(GetEnumConstants<uint8>(FlatSparse8), MakeConstArrayView({3})));
+			CHECK(EqualItems(GetEnumConstants<uint8>(FlagSparse8), MakeConstArrayView({2, 32})));
+			CHECK(GetEnumConstants<uint8>(FlatLimit8)[0] == 0xFF);
+			CHECK(GetEnumConstants<uint16>(FlatLimit16)[0] == 0xFFFF);
+			CHECK(GetEnumConstants<uint32>(FlatLimit32)[0] == 0xFFFFFFFF);
+			CHECK(GetEnumConstants<uint64>(FlatLimit64)[0] == 0xFFFFFFFFFFFFFFFF);
+			CHECK(GetEnumConstants<uint64>(FlagLimit64)[0] == 1);
+			CHECK(GetEnumConstants<uint64>(FlagLimit64)[1] == 0x8000000000000000);
 		});
 	}
 
