PlainProps rename  FArithmeticBindType to FBasicLeafBindType

#rnx

diff --git a/Engine/Plugins/Experimental/PlainProps/Source/Private/PlainPropsLoad.cpp b/Engine/Plugins/Experimental/PlainProps/Source/Private/PlainPropsLoad.cpp
--- a/Engine/Plugins/Experimental/PlainProps/Source/Private/PlainPropsLoad.cpp
+++ b/Engine/Plugins/Experimental/PlainProps/Source/Private/PlainPropsLoad.cpp
@@ -964,7 +964,7 @@ private:
 				SetBit(*Member, Leaf.Bitfield.Idx, BitIt.GrabNext(ByteIt));
 				break;
 			default:
-				switch (Leaf.Arithmetic.Width)
+				switch (Leaf.Basic.Width)
 				{
 					case ELeafWidth::B8:	FMemory::Memcpy(Member, ByteIt.GrabBytes(1), 1); break;
 					case ELeafWidth::B16:	FMemory::Memcpy(Member, ByteIt.GrabBytes(2), 2); break;
diff --git a/Engine/Plugins/Experimental/PlainProps/Source/Private/PlainPropsSave.cpp b/Engine/Plugins/Experimental/PlainProps/Source/Private/PlainPropsSave.cpp
--- a/Engine/Plugins/Experimental/PlainProps/Source/Private/PlainPropsSave.cpp
+++ b/Engine/Plugins/Experimental/PlainProps/Source/Private/PlainPropsSave.cpp
@@ -214,10 +214,10 @@ FBuiltRange* SaveLeafRange(const void* Range, const ILeafRangeBinding& Binding,
 	}
 }
 
-ELeafWidth GetArithmeticWidth(FLeafBindType Leaf)
+static ELeafWidth GetBasicWidth(FLeafBindType Leaf)
 {
 	checkf(Leaf.Bind.Type != ELeafBindType::BitfieldBool, TEXT("Arrays of bitfields is not a thing"));
-	return Leaf.Arithmetic.Width;
+	return Leaf.Basic.Width;
 }
 
 [[nodiscard]] FBuiltRange* SaveRange(const void* Range, FRangeMemberBinding Member, const FSaveContext& Ctx)
@@ -233,7 +233,7 @@ ELeafWidth GetArithmeticWidth(FLeafBindType Leaf)
 	const IItemRangeBinding& ItemBinding = Binding.AsItemBinding();
 	switch (InnerType.GetKind())
 	{
-	case EMemberKind::Leaf:		return SaveLeafRange(Range, ItemBinding, Ctx, GetArithmeticWidth(InnerType.AsLeaf()));
+	case EMemberKind::Leaf:		return SaveLeafRange(Range, ItemBinding, Ctx, GetBasicWidth(InnerType.AsLeaf()));
 	case EMemberKind::Range:	return SaveNonLeafRange<FNestedRangeSaver>(Range, ItemBinding, Ctx, GetInnerRange(Member));
 	case EMemberKind::Struct:	return SaveStructRange(Range, ItemBinding, Ctx, static_cast<FStructSchemaId>(Member.InnermostSchema.Get()));
 	}
@@ -471,7 +471,7 @@ static bool DiffItem(const uint8* A, const uint8* B, const FSaveContext& Ctx, FR
 	const IItemRangeBinding& ItemBinding = Binding.AsItemBinding();
 	switch (InnerType.GetKind())
 	{
-	case EMemberKind::Leaf:		return DiffItemRange(A, B, ItemBinding, Ctx, SizeOf(GetArithmeticWidth(InnerType.AsLeaf())));
+	case EMemberKind::Leaf:		return DiffItemRange(A, B, ItemBinding, Ctx, SizeOf(GetBasicWidth(InnerType.AsLeaf())));
 	case EMemberKind::Range:	return DiffItemRange(A, B, ItemBinding, Ctx, GetInnerRange(Member));
 	case EMemberKind::Struct:	return DiffItemRange(A, B, ItemBinding, Ctx, static_cast<FStructSchemaId>(Member.InnermostSchema.Get()));
 	}
diff --git a/Engine/Plugins/Experimental/PlainProps/Source/Public/PlainPropsBind.h b/Engine/Plugins/Experimental/PlainProps/Source/Public/PlainPropsBind.h
--- a/Engine/Plugins/Experimental/PlainProps/Source/Public/PlainPropsBind.h
+++ b/Engine/Plugins/Experimental/PlainProps/Source/Public/PlainPropsBind.h
@@ -50,14 +50,14 @@ inline static constexpr ELeafType ToLeafType(ELeafBindType Type)
 	return Type == ELeafBindType::BitfieldBool ? ELeafType::Bool : static_cast<ELeafType>(static_cast<uint8>(Type));
 }
 
-struct FArithmeticBindType
+struct FBasicLeafBindType
 {
     EMemberKind		_ : 2;
 	ELeafBindType	__ : 3;
 	ELeafWidth		Width : 2;
 };
 
-struct FBitfieldBindType
+struct FBitfieldBoolBindType
 {
     EMemberKind		_ : 2;
 	ELeafBindType	__ : 3;
@@ -66,10 +66,10 @@ struct FBitfieldBindType
 
 union FLeafBindType
 {
-	constexpr explicit FLeafBindType(ELeafBindType ArithmeticType, ELeafWidth Width) : Arithmetic({EMemberKind::Leaf, ArithmeticType, Width}) {}
-	constexpr explicit FLeafBindType(FUnpackedLeafType In) : Arithmetic({EMemberKind::Leaf, ToLeafBindType(In.Type), In.Width}) {}
+	constexpr explicit FLeafBindType(ELeafBindType BasicType, ELeafWidth Width) : Basic({EMemberKind::Leaf, BasicType, Width}) {}
+	constexpr explicit FLeafBindType(FUnpackedLeafType In) : Basic({EMemberKind::Leaf, ToLeafBindType(In.Type), In.Width}) {}
 	constexpr explicit FLeafBindType(FLeafType In) : FLeafBindType(FUnpackedLeafType(In)) {}
-	constexpr explicit FLeafBindType(FBitfieldBindType In) : Bitfield({EMemberKind::Leaf, ELeafBindType::BitfieldBool, In.Idx}) {}
+	constexpr explicit FLeafBindType(FBitfieldBoolBindType In) : Bitfield({EMemberKind::Leaf, ELeafBindType::BitfieldBool, In.Idx}) {}
 	constexpr explicit FLeafBindType(uint8 BitfieldIdx) : Bitfield({EMemberKind::Leaf, ELeafBindType::BitfieldBool, BitfieldIdx}) {}
 
 	struct
@@ -77,8 +77,8 @@ union FLeafBindType
 		EMemberKind			_ : 2;
 		ELeafBindType		Type : 3;
 	}						Bind;
-	FArithmeticBindType		Arithmetic;
-	FBitfieldBindType		Bitfield;
+	FBasicLeafBindType		Basic;
+	FBitfieldBoolBindType	Bitfield;
 };
 
 inline static constexpr FLeafType ToLeafType(FLeafBindType Leaf)
@@ -88,7 +88,7 @@ inline static constexpr FLeafType ToLeafType(FLeafBindType Leaf)
 		return { EMemberKind::Leaf, ELeafWidth::B8, ELeafType::Bool };
 	}
 	
-	return { EMemberKind::Leaf,  Leaf.Arithmetic.Width, ToLeafType(Leaf.Bind.Type) };
+	return { EMemberKind::Leaf,  Leaf.Basic.Width, ToLeafType(Leaf.Bind.Type) };
 }
 
 struct FRangeBindType : FRangeType {};
@@ -100,7 +100,7 @@ union FMemberBindType
 	constexpr explicit FMemberBindType(FLeafType In) : Leaf(In) {}
 	constexpr explicit FMemberBindType(FUnpackedLeafType In) : Leaf(In) {}
 	constexpr explicit FMemberBindType(FLeafBindType In) : Leaf(In) {}
-	constexpr explicit FMemberBindType(FBitfieldBindType In) : Leaf(In) {}
+	constexpr explicit FMemberBindType(FBitfieldBoolBindType In) : Leaf(In) {}
 	constexpr explicit FMemberBindType(FRangeType In) : Range(In) {}
 	constexpr explicit FMemberBindType(ERangeSizeType MaxSize) : Range({EMemberKind::Range, MaxSize}) {}
 	constexpr explicit FMemberBindType(FStructType In) : Struct(In) {}
@@ -155,13 +155,12 @@ struct FUnpackedLeafBindType
 		uint8				BitfieldIdx;
 	};
 	
-	//constexpr FUnpackedLeafBindType(ELeafBindType InType, ELeafWidth InWidth) : Type(InType), Width(InWidth) {}
 	constexpr FUnpackedLeafBindType(FLeafBindType In)
 	: Type(In.Bind.Type)
 	{
 		if (Type == ELeafBindType::BitfieldBool)
 		{
-			Width = In.Arithmetic.Width;
+			Width = In.Basic.Width;
 		}
 		else
 		{
@@ -169,7 +168,6 @@ struct FUnpackedLeafBindType
 		}
 	}
 
-//	constexpr bool operator==(FUnpackedLeafBindType O) const { return Type == O.Type && Width == O.Width; }
 	FMemberBindType Pack() const
 	{ 
 		return FMemberBindType(Type == ELeafBindType::BitfieldBool ? FLeafBindType(BitfieldIdx) : FLeafBindType(Type, Width));
