PlainProps: Make FUnpackedLeafType::AsInt() inline implementation more compact and readable by replacing Memcpy with BitCast.

#rnx

diff --git a/Engine/Plugins/Experimental/PlainProps/Source/Public/PlainPropsTypes.h b/Engine/Plugins/Experimental/PlainProps/Source/Public/PlainPropsTypes.h
--- a/Engine/Plugins/Experimental/PlainProps/Source/Public/PlainPropsTypes.h
+++ b/Engine/Plugins/Experimental/PlainProps/Source/Public/PlainPropsTypes.h
@@ -4,6 +4,7 @@
 
 #include "Containers/ArrayView.h"
 #include "Misc/AssertionMacros.h"
+#include "Templates/TypeCompatibleBytes.h"
 #include "Templates/TypeHash.h"
 #include <type_traits>
 
@@ -342,15 +343,9 @@ struct FUnpackedLeafType
 	inline constexpr FUnpackedLeafType(ELeafType InType, ELeafWidth InWidth) : Type(InType), Width(InWidth) {}
 	inline constexpr FUnpackedLeafType(FLeafType In) : Type(In.Type), Width(In.Width) {}
 
-	inline bool operator==(FUnpackedLeafType O) const { return AsInt() == O.AsInt(); }
-	inline constexpr FMemberType Pack() const { return FMemberType(Type, Width); }
-
-	inline uint16 AsInt() const
-	{
-		uint16 Out = 0;
-		FMemory::Memcpy(&Out, this, sizeof(uint16));
-		return Out;
-	}
+	inline uint16 AsInt() const							{ return BitCast<uint16>(*this); }
+	inline bool operator==(FUnpackedLeafType O) const	{ return AsInt() == O.AsInt(); }
+	inline constexpr FMemberType Pack() const			{ return FMemberType(Type, Width); }
 };
 
 template<Enumeration T>
