PlainProps: Introduce a PlainPropsPrint public header with initial ToString() functions for FUnpackedLeafType and ERangeSizeType.

#rnx

diff --git a/Engine/Plugins/Experimental/PlainProps/Source/Private/PlainPropsBuildSchema.cpp b/Engine/Plugins/Experimental/PlainProps/Source/Private/PlainPropsBuildSchema.cpp
--- a/Engine/Plugins/Experimental/PlainProps/Source/Private/PlainPropsBuildSchema.cpp
+++ b/Engine/Plugins/Experimental/PlainProps/Source/Private/PlainPropsBuildSchema.cpp
@@ -3,34 +3,20 @@
 #include "PlainPropsBuildSchema.h"
 #include "PlainPropsInternalBuild.h"
 #include "PlainPropsInternalFormat.h"
+#include "PlainPropsPrint.h"
 #include "Algo/Compare.h"
 #include "Algo/Find.h"
 #include "Containers/Map.h"
+#include "Containers/StringConv.h"
 
 namespace PlainProps
 {
 
-static const TCHAR* ToString(FLeafType Leaf)
-{
-	const TCHAR* Leaves[8][4] = {
-		{TEXT("bool"),		TEXT("ERR_b16"),	TEXT("ERR_b32"),	TEXT("ERR_b64")},
-		{TEXT("int8"),		TEXT("int16"),		TEXT("int32"),		TEXT("int64")},
-		{TEXT("uint8"),		TEXT("uint16"),		TEXT("uint32"),		TEXT("uint64")},
-		{TEXT("ERR_fp8"),	TEXT("ERR_fp16"),	TEXT("float"),		TEXT("double")},
-		{TEXT("hex8"),		TEXT("hex16"),		TEXT("hex32"),		TEXT("hex64")},
-		{TEXT("enum8"),		TEXT("enum16"),		TEXT("enum32"),		TEXT("enum64")},
-		{TEXT("utf8"),		TEXT("utf16"),		TEXT("utf32"),		TEXT("ERR_utf64")},
-		{TEXT("ERR_oob"),	TEXT("ERR_oob"),	TEXT("ERR_oob"),	TEXT("ERR_oob")}};
-
-	return Leaves[(uint8)Leaf.Type][(uint8)Leaf.Width];
-}
-
-
 static FString PrintMemberSchema(FMemberType Type, FOptionalSchemaId InnerSchema, TConstArrayView<FMemberType> InnerRangeTypes)
 {
 	switch (Type.GetKind())
 	{
-	case EMemberKind::Leaf:		return ToString(Type.AsLeaf());
+	case EMemberKind::Leaf:		return StringCast<TCHAR>(ToString(Type.AsLeaf()).GetData()).Get();
 	case EMemberKind::Struct:	return FString::Printf(TEXT("Struct [%d]%s%s"), InnerSchema.Get().Idx, Type.AsStruct().IsSuper ? TEXT(" (super)") : TEXT(""), Type.AsStruct().IsDynamic ? TEXT(" (dynamic)") : TEXT(""));
 	case EMemberKind::Range:	return TEXT("Range of ") + PrintMemberSchema(InnerRangeTypes[0], InnerSchema, InnerRangeTypes.RightChop(1));
 	}
diff --git a/Engine/Plugins/Experimental/PlainProps/Source/Public/PlainPropsPrint.h b/Engine/Plugins/Experimental/PlainProps/Source/Public/PlainPropsPrint.h
new file mode 100644
--- /dev/null
+++ b/Engine/Plugins/Experimental/PlainProps/Source/Public/PlainPropsPrint.h
@@ -0,0 +1,32 @@
+// Copyright Epic Games, Inc. All Rights Reserved.
+
+#pragma once
+
+#include "PlainPropsTypes.h"
+#include "Containers/StringView.h"
+
+namespace PlainProps
+{
+
+inline FAnsiStringView ToString(ERangeSizeType Width)
+{
+	static const FAnsiStringView Widths[] = { "Uni", "S8", "U8", "S16", "U16", "S32", "U32", "S64", "U64" };
+	return Widths[(uint8)Width];
+}
+
+inline FAnsiStringView ToString(FUnpackedLeafType Leaf)
+{
+	static const FAnsiStringView Leaves[8][4] = {
+		{("bool"),		("ERR_b16"),	("ERR_b32"),	("ERR_b64")},
+		{("i8"),		("i16"),		("i32"),		("i64")},
+		{("u8"),		("u16"),		("u32"),		("u64")},
+		{("ERR_fp8"),	("ERR_fp16"),	("float"),		("double")},
+		{("hex8"),		("hex16"),		("hex32"),		("hex64")},
+		{("enum8"),		("enum16"),		("enum32"),		("enum64")},
+		{("utf8"),		("utf16"),		("utf32"),		("ERR_utf64")},
+		{("ERR_oob"),	("ERR_oob"),	("ERR_oob"),	("ERR_oob")}};
+
+	return Leaves[(uint8)Leaf.Type][(uint8)Leaf.Width];
+}
+
+} // namespace PlainProps
