PlainProps moved some code to new internal header

#rnx

diff --git a/Engine/Plugins/Experimental/PlainProps/Source/Private/PlainPropsBind.cpp b/Engine/Plugins/Experimental/PlainProps/Source/Private/PlainPropsBind.cpp
--- a/Engine/Plugins/Experimental/PlainProps/Source/Private/PlainPropsBind.cpp
+++ b/Engine/Plugins/Experimental/PlainProps/Source/Private/PlainPropsBind.cpp
@@ -2,6 +2,7 @@
 
 #include "PlainPropsBind.h"
 #include "PlainPropsIndex.h"
+#include "PlainPropsInternalBind.h"
 #include "PlainPropsInternalBuild.h"
 #include "PlainPropsInternalFormat.h"
 #include "PlainPropsInternalRead.h"
diff --git a/Engine/Plugins/Experimental/PlainProps/Source/Private/PlainPropsInternalBind.h b/Engine/Plugins/Experimental/PlainProps/Source/Private/PlainPropsInternalBind.h
new file mode 100644
--- /dev/null
+++ b/Engine/Plugins/Experimental/PlainProps/Source/Private/PlainPropsInternalBind.h
@@ -0,0 +1,100 @@
+// Copyright Epic Games, Inc. All Rights Reserved.
+
+#pragma once
+
+#include "PlainPropsBind.h"
+
+namespace PlainProps 
+{
+
+// Iterates over member bindings
+class FMemberVisitor
+{
+public:
+	explicit FMemberVisitor(const FSchemaBinding& InSchema);
+
+	bool						HasMore() const			{ return MemberIdx < NumMembers; }
+	uint16						GetIndex() const		{ return MemberIdx; }
+	
+	EMemberKind					PeekKind() const;		// @pre HasMore()
+	FMemberBindType				PeekType() const;		// @pre HasMore()
+	uint32						PeekOffset() const;		// @pre HasMore()
+
+	FLeafMemberBinding			GrabLeaf();				// @pre PeekKind() == EMemberKind::Leaf
+	FRangeMemberBinding			GrabRange();			// @pre PeekKind() == EMemberKind::Range
+	FStructMemberBinding		GrabStruct();			// @pre PeekKind() == EMemberKind::Struct
+	void						SkipMember();
+
+protected: // for unit tests
+	const FSchemaBinding&		Schema;
+	const uint16				NumMembers;
+	uint16						MemberIdx = 0;
+	uint16						InnerRangeIdx = 0;		// Types of [nested] ranges
+	uint16						InnerSchemaIdx = 0;		// Types of static structs and enums
+
+	using FMemberBindTypeRange = TConstArrayView<FMemberBindType>;
+
+	uint64						GrabMemberOffset();
+	FMemberBindTypeRange		GrabInnerTypes();
+	FSchemaId					GrabInnerSchema();
+	FStructSchemaId				GrabStructSchema(FStructType Type);
+	FOptionalSchemaId			GrabRangeSchema(FMemberType InnermostType);
+	FEnumSchemaId				GrabEnumSchema()			{ return static_cast<FEnumSchemaId&&>(GrabInnerSchema()); }
+};
+
+////////////////////////////////////////////////////////////////////////////////////////////////
+
+struct FMemberBinder
+{
+	FMemberBinder(FSchemaBinding& InSchema)
+	: Schema(InSchema)
+	, MemberIt(Schema.Members)
+	, RangeTypeIt(const_cast<FMemberBindType*>(Schema.GetInnerRangeTypes()))
+	, OffsetIt(const_cast<uint32*>(Schema.GetOffsets()))
+	, InnerSchemaIt(const_cast<FSchemaId*>(Schema.GetInnerSchemas()))
+	, RangeBindingIt(const_cast<FRangeBinding*>(Schema.GetRangeBindings()))
+	{}
+
+	~FMemberBinder()
+	{
+		check(MemberIt == Schema.GetInnerRangeTypes());
+		check(Align(RangeTypeIt, alignof(uint32)) == (const void*)Schema.GetOffsets());
+		check(OffsetIt == (const void*)Schema.GetInnerSchemas());
+		check(Align(InnerSchemaIt, alignof(FRangeBinding)) == (const void*)Schema.GetRangeBindings() || Schema.NumInnerRanges == 0);
+		check(Schema.NumInnerRanges == RangeBindingIt - Schema.GetRangeBindings());
+	}
+
+	void AddMember(FMemberBindType Type, uint32 Offset)
+	{
+		*MemberIt++ = Type;
+		*OffsetIt++ = Offset;
+	}
+
+	void AddRange(TConstArrayView<FRangeBinding> Ranges, FMemberBindType InnermostType, uint32 Offset)
+	{
+		AddMember(FMemberBindType(Ranges[0].GetSizeType()), Offset);
+
+		for (FRangeBinding Range : Ranges.RightChop(1))
+		{
+			*RangeTypeIt++ = FMemberBindType(Range.GetSizeType());
+		}
+		*RangeTypeIt++ = InnermostType;
+
+		FMemory::Memcpy(RangeBindingIt, Ranges.GetData(), Ranges.Num() * Ranges.GetTypeSize());
+		RangeBindingIt += Ranges.Num();
+	}
+	
+	void AddInnerSchema(FSchemaId InnermostSchema)
+	{
+		*InnerSchemaIt++ = InnermostSchema;
+	}
+
+	FSchemaBinding& Schema;
+	FMemberBindType* MemberIt;
+	FMemberBindType* RangeTypeIt;
+	uint32* OffsetIt;
+	FSchemaId* InnerSchemaIt;
+	FRangeBinding* RangeBindingIt;
+};
+
+} // namespace PlainProps
\ No newline at end of file
diff --git a/Engine/Plugins/Experimental/PlainProps/Source/Private/PlainPropsLoad.cpp b/Engine/Plugins/Experimental/PlainProps/Source/Private/PlainPropsLoad.cpp
--- a/Engine/Plugins/Experimental/PlainProps/Source/Private/PlainPropsLoad.cpp
+++ b/Engine/Plugins/Experimental/PlainProps/Source/Private/PlainPropsLoad.cpp
@@ -3,6 +3,7 @@
 #include "PlainPropsLoad.h"
 #include "PlainPropsLoadMember.h"
 #include "PlainPropsBind.h"
+#include "PlainPropsInternalBind.h"
 #include "PlainPropsInternalFormat.h"
 #include "PlainPropsInternalRead.h"
 #include "Algo/Find.h"
diff --git a/Engine/Plugins/Experimental/PlainProps/Source/Private/PlainPropsSave.cpp b/Engine/Plugins/Experimental/PlainProps/Source/Private/PlainPropsSave.cpp
--- a/Engine/Plugins/Experimental/PlainProps/Source/Private/PlainPropsSave.cpp
+++ b/Engine/Plugins/Experimental/PlainProps/Source/Private/PlainPropsSave.cpp
@@ -2,6 +2,7 @@
 
 #include "PlainPropsSave.h"
 #include "PlainPropsBind.h"
+#include "PlainPropsInternalBind.h"
 #include "PlainPropsInternalBuild.h"
 #include "PlainPropsInternalFormat.h"
 #include "PlainPropsSaveMember.h"
diff --git a/Engine/Plugins/Experimental/PlainProps/Source/Public/PlainPropsBind.h b/Engine/Plugins/Experimental/PlainProps/Source/Public/PlainPropsBind.h
--- a/Engine/Plugins/Experimental/PlainProps/Source/Public/PlainPropsBind.h
+++ b/Engine/Plugins/Experimental/PlainProps/Source/Public/PlainPropsBind.h
@@ -205,41 +205,6 @@ struct FStructMemberBinding
 	SIZE_T					Offset;
 };
 
-// Iterates over member bindings
-class FMemberVisitor
-{
-public:
-	explicit FMemberVisitor(const FSchemaBinding& InSchema);
-
-	bool						HasMore() const			{ return MemberIdx < NumMembers; }
-	uint16						GetIndex() const		{ return MemberIdx; }
-	
-	EMemberKind					PeekKind() const;		// @pre HasMore()
-	FMemberBindType				PeekType() const;		// @pre HasMore()
-	uint32						PeekOffset() const;		// @pre HasMore()
-
-	FLeafMemberBinding			GrabLeaf();				// @pre PeekKind() == EMemberKind::Leaf
-	FRangeMemberBinding			GrabRange();			// @pre PeekKind() == EMemberKind::Range
-	FStructMemberBinding		GrabStruct();			// @pre PeekKind() == EMemberKind::Struct
-	void						SkipMember();
-
-protected: // for unit tests
-	const FSchemaBinding&		Schema;
-	const uint16				NumMembers;
-	uint16						MemberIdx = 0;
-	uint16						InnerRangeIdx = 0;		// Types of [nested] ranges
-	uint16						InnerSchemaIdx = 0;		// Types of static structs and enums
-
-	using FMemberBindTypeRange = TConstArrayView<FMemberBindType>;
-
-	uint64						GrabMemberOffset();
-	FMemberBindTypeRange		GrabInnerTypes();
-	FSchemaId					GrabInnerSchema();
-	FStructSchemaId				GrabStructSchema(FStructType Type);
-	FOptionalSchemaId			GrabRangeSchema(FMemberType InnermostType);
-	FEnumSchemaId				GrabEnumSchema()			{ return static_cast<FEnumSchemaId&&>(GrabInnerSchema()); }
-};
-
 ////////////////////////////////////////////////////////////////////////////////////////////////
 
 struct FDualStructSchemaId
@@ -1130,61 +1095,6 @@ void BindNativeStruct(FSchemaBindings& Out, FStructSchemaId BindId, FStructSchem
 	Out.BindStruct(BindId, DeclId, MemberBindings);
 }
 
-////////////////////////////////////////////////////////////////////////////////////////////////
-
-struct FMemberBinder
-{
-	FMemberBinder(FSchemaBinding& InSchema)
-	: Schema(InSchema)
-	, MemberIt(Schema.Members)
-	, RangeTypeIt(const_cast<FMemberBindType*>(Schema.GetInnerRangeTypes()))
-	, OffsetIt(const_cast<uint32*>(Schema.GetOffsets()))
-	, InnerSchemaIt(const_cast<FSchemaId*>(Schema.GetInnerSchemas()))
-	, RangeBindingIt(const_cast<FRangeBinding*>(Schema.GetRangeBindings()))
-	{}
-
-	~FMemberBinder()
-	{
-		check(MemberIt == Schema.GetInnerRangeTypes());
-		check(Align(RangeTypeIt, alignof(uint32)) == (const void*)Schema.GetOffsets());
-		check(OffsetIt == (const void*)Schema.GetInnerSchemas());
-		check(Align(InnerSchemaIt, alignof(FRangeBinding)) == (const void*)Schema.GetRangeBindings() || Schema.NumInnerRanges == 0);
-		check(Schema.NumInnerRanges == RangeBindingIt - Schema.GetRangeBindings());
-	}
-
-	void AddMember(FMemberBindType Type, uint32 Offset)
-	{
-		*MemberIt++ = Type;
-		*OffsetIt++ = Offset;
-	}
-
-	void AddRange(TConstArrayView<FRangeBinding> Ranges, FMemberBindType InnermostType, uint32 Offset)
-	{
-		AddMember(FMemberBindType(Ranges[0].GetSizeType()), Offset);
-
-		for (FRangeBinding Range : Ranges.RightChop(1))
-		{
-			*RangeTypeIt++ = FMemberBindType(Range.GetSizeType());
-		}
-		*RangeTypeIt++ = InnermostType;
-
-		FMemory::Memcpy(RangeBindingIt, Ranges.GetData(), Ranges.Num() * Ranges.GetTypeSize());
-		RangeBindingIt += Ranges.Num();
-	}
-	
-	void AddInnerSchema(FSchemaId InnermostSchema)
-	{
-		*InnerSchemaIt++ = InnermostSchema;
-	}
-
-	FSchemaBinding& Schema;
-	FMemberBindType* MemberIt;
-	FMemberBindType* RangeTypeIt;
-	uint32* OffsetIt;
-	FSchemaId* InnerSchemaIt;
-	FRangeBinding* RangeBindingIt;
-};
-
 //////////////////////////////////////////////////////////////////////////
 
 // Save -> load struct ids for ESchemaFormat::InMemoryNames
