PlainProps: Initial drop of printing a read schema to a yaml utf8 string containing all schemas and objects with values and debug type info for all fields.

#rnx

diff --git a/Engine/Plugins/Experimental/PlainProps/Source/Private/PlainPropsBuildSchema.cpp b/Engine/Plugins/Experimental/PlainProps/Source/Private/PlainPropsBuildSchema.cpp
--- a/Engine/Plugins/Experimental/PlainProps/Source/Private/PlainPropsBuildSchema.cpp
+++ b/Engine/Plugins/Experimental/PlainProps/Source/Private/PlainPropsBuildSchema.cpp
@@ -3,29 +3,46 @@
 #include "PlainPropsBuildSchema.h"
 #include "PlainPropsInternalBuild.h"
 #include "PlainPropsInternalFormat.h"
-#include "PlainPropsPrint.h"
+#include "PlainPropsInternalPrint.h"
 #include "Algo/Compare.h"
 #include "Algo/Find.h"
 #include "Containers/Map.h"
 #include "Containers/StringConv.h"
+#include "Misc/StringBuilder.h"
 
 namespace PlainProps
 {
 
-static FString PrintMemberSchema(FMemberType Type, FOptionalSchemaId InnerSchema, TConstArrayView<FMemberType> InnerRangeTypes)
+static FString PrintMemberSchema(const FDebugIds& Debug, FMemberType Type, FOptionalSchemaId InnerSchema, TConstArrayView<FMemberType> InnerRangeTypes)
 {
-	switch (Type.GetKind())
+	TUtf8StringBuilder<256> Utf8SchemaStr;
+	PrintMemberSchema(Utf8SchemaStr, Debug, Type, InnerSchema, InnerRangeTypes);
+	FString SchemaStr(FStringView(StringCast<TCHAR>(Utf8SchemaStr.GetData(), Utf8SchemaStr.Len())));
+
+	FMemberType InnermostType = InnerRangeTypes.Num() > 0 ? InnerRangeTypes.Last() : Type;
+	if (InnermostType.IsStruct())
+	{
+		return FString::Printf(TEXT("%sstruct [%d]%s%s => %s"),
+			Type.IsRange() ? TEXT("Range(s) of ") : TEXT(""),
+			InnerSchema.Get().Idx,
+			InnermostType.AsStruct().IsSuper ? TEXT(" (super)") : TEXT(""),
+			InnermostType.AsStruct().IsDynamic ? TEXT(" (dynamic)") : TEXT(""),
+			*SchemaStr);
+	}
+	else if (IsEnum(InnermostType))
 	{
-	case EMemberKind::Leaf:		return StringCast<TCHAR>(ToString(Type.AsLeaf()).GetData()).Get();
-	case EMemberKind::Struct:	return FString::Printf(TEXT("Struct [%d]%s%s"), InnerSchema.Get().Idx, Type.AsStruct().IsSuper ? TEXT(" (super)") : TEXT(""), Type.AsStruct().IsDynamic ? TEXT(" (dynamic)") : TEXT(""));
-	case EMemberKind::Range:	return TEXT("Range of ") + PrintMemberSchema(InnerRangeTypes[0], InnerSchema, InnerRangeTypes.RightChop(1));
+		return FString::Printf(TEXT("%s%s [%d] => %s"),
+			Type.IsRange() ? TEXT("Range(s) of ") : TEXT(""),
+			StringCast<TCHAR>(ToString(InnermostType.AsLeaf()).GetData()).Get(),
+			InnerSchema.Get().Idx,
+			*SchemaStr);
 	}
-	return "Illegal member kind";
+	return SchemaStr;
 }
 
-static FString PrintMemberSchema(FMemberSchema Schema)
+static FString PrintMemberSchema(const FDebugIds& Debug, FMemberSchema Schema)
 {
-	return PrintMemberSchema(Schema.Type, Schema.InnerSchema, Schema.GetInnerRangeTypes());
+	return PrintMemberSchema(Debug, Schema.Type, Schema.InnerSchema, Schema.GetInnerRangeTypes());
 }
 
 //////////////////////////////////////////////////////////////////////////
@@ -229,8 +246,11 @@ void FStructSchemaBuilder::NoteMembersRecursively(const FBuiltStruct& Struct)
 			}
 			else
 			{
-				checkf(NotedSchema == Member.Schema, TEXT("Member '%s' in '%s' first added as %s and later as %s."),
-					*AllSchemas.GetDebug().Print(Member.Name), *AllSchemas.GetDebug().Print(Declaration.Type), *PrintMemberSchema(NotedSchema), *PrintMemberSchema(Member.Schema));
+				checkf(NotedSchema == Member.Schema, TEXT("Member '%s' in '%s' first added as '%s' and later as '%s'."),
+					*AllSchemas.GetDebug().Print(Member.Name),
+					*AllSchemas.GetDebug().Print(Declaration.Type),
+					*PrintMemberSchema(AllSchemas.GetDebug(), NotedSchema),
+					*PrintMemberSchema(AllSchemas.GetDebug(), Member.Schema));
 			}
 		}
 		else
diff --git a/Engine/Plugins/Experimental/PlainProps/Source/Private/PlainPropsInternalFormat.h b/Engine/Plugins/Experimental/PlainProps/Source/Private/PlainPropsInternalFormat.h
--- a/Engine/Plugins/Experimental/PlainProps/Source/Private/PlainPropsInternalFormat.h
+++ b/Engine/Plugins/Experimental/PlainProps/Source/Private/PlainPropsInternalFormat.h
@@ -52,7 +52,8 @@ static_assert(offsetof(FSchemaBatch, NumStructSchemas) == 16,					"Add binary fo
 
 enum class ESuper : uint8 { No, Unused, Used, Reused };
 
-inline bool UsesSuper(ESuper Inheritance) { return Inheritance == ESuper::Used || Inheritance == ESuper::Reused; }
+inline bool UsesSuper(ESuper Inheritance)				{ return Inheritance == ESuper::Used   || Inheritance == ESuper::Reused; }
+inline bool SkipDeclaredSuperSchema(ESuper Inheritance)	{ return Inheritance == ESuper::Unused || Inheritance == ESuper::Used; }
 
 struct FStructSchema
 {
diff --git a/Engine/Plugins/Experimental/PlainProps/Source/Private/PlainPropsInternalPrint.h b/Engine/Plugins/Experimental/PlainProps/Source/Private/PlainPropsInternalPrint.h
new file mode 100644
--- /dev/null
+++ b/Engine/Plugins/Experimental/PlainProps/Source/Private/PlainPropsInternalPrint.h
@@ -0,0 +1,53 @@
+// Copyright Epic Games, Inc. All Rights Reserved.
+
+#pragma once
+
+#include "Containers/ArrayView.h"
+#include "Misc/StringBuilder.h"
+#include "PlainPropsPrint.h"
+#include "Templates/UniquePtr.h"
+
+namespace PlainProps
+{
+
+class FYamlBuilder;
+struct FSchemaBatch;
+struct FEnumSchema;
+struct FStructSchema;
+
+///////////////////////////////////////////////////////////////////////////////
+
+struct FYamlBuilderDeleter
+{
+	PLAINPROPS_API void operator()(FYamlBuilder* Ptr) const;
+};
+using FYamlBuilderPtr = TUniquePtr<FYamlBuilder, FYamlBuilderDeleter>;
+
+[[nodiscard]] FYamlBuilderPtr MakeYamlBuilder(FUtf8StringBuilderBase& StringBuilder);
+
+///////////////////////////////////////////////////////////////////////////////
+
+class FBatchPrinter
+{
+public:
+	FBatchPrinter(FYamlBuilder& InTextBuilder, const FDebugIds& InNames);
+	~FBatchPrinter();
+
+	void PrintSchemas(FReadBatchId Batch);
+	void PrintObjects(TConstArrayView<FStructView> Structs);
+
+private:
+	void PrintSchemas(const FSchemaBatch& Batch, FReadBatchId BatchId);
+	void PrintStructSchema(const FStructSchema& Struct, FReadBatchId BatchId);
+	void PrintEnumSchema(const FEnumSchema& Enum);
+
+	template<typename IntType>
+	void PrintEnumConstants(TConstArrayView<FNameId> EnumNames, TConstArrayView<IntType> Constants, bool bFlagMode);
+
+	FYamlBuilder& TextBuilder;
+	const FDebugIds& Names;
+};
+
+///////////////////////////////////////////////////////////////////////////////
+
+} // namespace PlainProps
diff --git a/Engine/Plugins/Experimental/PlainProps/Source/Private/PlainPropsInternalRead.h b/Engine/Plugins/Experimental/PlainProps/Source/Private/PlainPropsInternalRead.h
--- a/Engine/Plugins/Experimental/PlainProps/Source/Private/PlainPropsInternalRead.h
+++ b/Engine/Plugins/Experimental/PlainProps/Source/Private/PlainPropsInternalRead.h
@@ -91,6 +91,33 @@ inline TSchemaRange<const FStructSchema>	GetStructSchemas(const FSchemaBatch& Ba
 
 //////////////////////////////////////////////////////////////////////////
 
+const FSchemaBatch& GetReadSchemas(FReadBatchId Batch);
+
+inline const FStructSchema& ResolveStructSchema(const FSchemaBatch& Batch, FSchemaId Id)
+{
+	check(Id.Idx < Batch.NumSchemas);
+	return *reinterpret_cast<const FStructSchema*>(reinterpret_cast<const uint8*>(&Batch) + Batch.SchemaOffsets[Id.Idx]);
+}
+
+inline const FEnumSchema& ResolveEnumSchema(const FSchemaBatch& Batch, FSchemaId Id)
+{
+	check(Id.Idx < Batch.NumSchemas);
+	return *reinterpret_cast<const FEnumSchema*>(reinterpret_cast<const uint8*>(&Batch) + Batch.SchemaOffsets[Id.Idx]);
+}
+
+inline FNestedScope ResolveNestedScope(const FSchemaBatch& Batch, FNestedScopeId Id)
+{
+	return Batch.GetNestedScopes()[Id.Idx];
+}
+
+inline FParametricTypeView ResolveParametricType(const FSchemaBatch& Batch, FParametricTypeId Id)
+{
+	FParametricType Type = Batch.GetParametricTypes()[Id.Idx];
+	return {Type.Name, IntCastChecked<uint8>(Type.Parameters.NumParameters), Batch.GetFirstParameter() + Type.Parameters.Idx};
+}
+
+//////////////////////////////////////////////////////////////////////////
+
 inline FMemoryView GetSchemaData(const FSchemaBatch& Batch)
 {
 	check(Batch.NumSchemas > 0);
diff --git a/Engine/Plugins/Experimental/PlainProps/Source/Private/PlainPropsPrint.cpp b/Engine/Plugins/Experimental/PlainProps/Source/Private/PlainPropsPrint.cpp
--- a/Engine/Plugins/Experimental/PlainProps/Source/Private/PlainPropsPrint.cpp
+++ b/Engine/Plugins/Experimental/PlainProps/Source/Private/PlainPropsPrint.cpp
@@ -2,14 +2,26 @@
 
 #include "PlainPropsPrint.h"
 #include "PlainPropsIndex.h"
+#include "Containers/StringConv.h"
+#include "Containers/Utf8String.h"
+#include "PlainPropsInternalFormat.h"
+#include "PlainPropsInternalPrint.h"
+#include "PlainPropsInternalRead.h"
+#include "Misc/StringBuilder.h"
 
 namespace PlainProps
 {
 
+static constexpr bool PrintWithComments = true;
+
+//////////////////////////////////////////////////////////////////////////
+
 struct FLiterals
 {
+	FAnsiStringView Super = "__super__";
+	FAnsiStringView Dynamic = "__dynamic__";
 	FAnsiStringView Ranges[9] = { 
-		("Uni"), ("S8"), ("U8"), ("S16"), ("U16"), ("S32"), ("U32"), ("S64"), ("U64") };
+		("Uni"), ("i8"), ("u8"), ("i16"), ("u16"), ("i32"), ("u32"), ("i64"), ("u64") };
 	FAnsiStringView Leaves[8][4] = {
 		{("bool"),		("ERR_b16"),	("ERR_b32"),	("ERR_b64")},
 		{("i8"),		("i16"),		("i32"),		("i64")},
@@ -52,4 +64,1165 @@ void FIdIndexerBase::InitParameterNames()
 	}
 }
 
-} // namespace PlainProps
\ No newline at end of file
+///////////////////////////////////////////////////////////////////////////////
+
+template<typename IntType>
+static TConstArrayView<IntType> GetConstants(const FEnumSchema& Schema)
+{
+	check(Schema.Width == ReflectLeaf<IntType>.Width);
+	return MakeConstArrayView<IntType>(
+		reinterpret_cast<const IntType*>(Schema.Footer + Schema.Num),
+		Schema.ExplicitConstants ? Schema.Num : 0);
+}
+
+///////////////////////////////////////////////////////////////////////////////
+
+FUtf8StringView PrintMemberSchemaImpl(
+	FUtf8StringBuilderBase& Out,
+	const FDebugIds& Names,
+	FMemberType Type,
+	FOptionalSchemaId InnerSchema = NoId,
+	TConstArrayView<FMemberType> InnerRangeTypes = {})
+{
+	FMemberType InnermostType = InnerRangeTypes.Num() > 0 ? InnerRangeTypes.Last() : Type;
+	check(!InnermostType.IsRange());
+	if (InnerSchema)
+	{
+		check(IsStructOrEnum(InnermostType));
+		if (InnermostType.IsLeaf())
+		{
+			check(IsEnum(InnermostType));
+			Names.AppendString(Out, static_cast<FEnumSchemaId>(InnerSchema.Get()));
+		}
+		else
+		{
+			check(InnermostType.IsStruct());
+			Names.AppendString(Out, static_cast<FStructSchemaId>(InnerSchema.Get()));
+		}
+	}
+	else if (InnermostType.IsLeaf())
+	{
+		check(!IsEnum(InnermostType));
+		Out.Append(ToString(InnermostType.AsLeaf()));
+	}
+	else if (InnermostType.IsStruct())
+	{
+		check(InnermostType.AsStruct().IsDynamic);
+		Out.Append(GLiterals.Dynamic);
+	}
+
+	if (Type.IsRange())
+	{
+		Out.AppendChar('[');
+		Out.Append(ToString(Type.AsRange().MaxSize));
+		Out.AppendChar(']');
+		for (FMemberType MemberType : InnerRangeTypes.LeftChop(1))
+		{
+			Out.AppendChar('[');
+			Out.Append(ToString(MemberType.AsRange().MaxSize));
+			Out.AppendChar(']');
+		}
+	}
+	return Out.ToView();
+}
+
+void PrintMemberSchema(
+	FUtf8StringBuilderBase& Out,
+	const FDebugIds& Names, 
+	FMemberType Type,
+	FOptionalSchemaId InnerSchema,
+	TConstArrayView<FMemberType> InnerRangeTypes)
+{
+	PrintMemberSchemaImpl(Out, Names, Type, InnerSchema, InnerRangeTypes);
+}
+
+///////////////////////////////////////////////////////////////////////////////
+
+struct FMemberSchemaView
+{
+	using FMemberTypeRange = TConstArrayView<FMemberType>;
+
+	FMemberType				Type;
+	FReadBatchId			Batch;
+	FOptionalSchemaId		InnerSchema;
+	FMemberTypeRange		InnerRangeTypes;
+
+	FMemberType GetInnermostType() const
+	{
+		return InnerRangeTypes.Num() > 0 ? InnerRangeTypes.Last() : Type;
+	}
+};
+
+///////////////////////////////////////////////////////////////////////////////
+
+class FStructSchemaReader
+{
+public:
+	FStructSchemaReader(const FStructSchema& Schema, FReadBatchId InBatch);
+
+	FTypeId					GetTypeId() const		{ return TypeId; }
+	bool					IsDense() const			{ return bIsDense; }
+	bool					HasSuper() const		{ return bHasSuper; }
+
+	bool					HasMore() const			{ return MemberIdx < NumMembers; }
+	
+	FOptionalMemberId		PeekName() const;		// @pre HasMore()
+	EMemberKind				PeekKind() const;		// @pre HasMore()
+	FMemberType				PeekType() const;		// @pre HasMore()
+
+	FStructSchemaHandle		GetSuper() const;		// @pre HasSuper()
+	FMemberSchemaView 		GrabMember();			// @pre HasMore()
+
+private:
+	const FMemberType*		Footer;
+	const FReadBatchId		Batch;					// Needed to resolve schemas
+	const FTypeId			TypeId;
+	const bool				bIsDense : 1;
+	const bool				bHasSuper : 1;
+	const bool				bUsesSuper : 1;
+	const uint32			NumMembers;
+	const uint32			NumRangeTypes;			// Number of ranges and nested ranges
+	const uint32			NumInnerSchemas;		// Number of static structs and enums
+
+	uint32					MemberIdx = 0;
+	uint32					RangeTypeIdx = 0;		// Types of [nested] ranges
+	uint32					InnerSchemaIdx = 0;		// Types of static structs and enums
+
+	void					AdvanceToNextMember()	{ ++MemberIdx; }
+
+	using FMemberTypeRange = TConstArrayView<FMemberType>;
+
+	FMemberTypeRange		GrabRangeTypes();
+	FSchemaId				GrabInnerSchema();
+	FOptionalSchemaId		GrabLeafSchema(FLeafType Leaf);
+	FOptionalSchemaId		GrabStructSchema(FStructType Struct);
+	FOptionalSchemaId		GrabRangeSchema(FMemberType InnermostType);
+
+	const FMemberType*		GetMemberTypes() const;
+	const FMemberType*		GetRangeTypes() const;
+	const FSchemaId*		GetInnerSchemas() const;
+	const FMemberId*		GetMemberNames() const;
+};
+
+///////////////////////////////////////////////////////////////////////////////
+
+class FYamlBuilder
+{
+public:
+	FYamlBuilder(FUtf8StringBuilderBase& InStringBuilder);
+	~FYamlBuilder();
+
+	void BeginDocument();
+	void EndDocument();
+
+	void BeginStruct(FUtf8StringView Id);
+	void BeginStruct();
+	void EndStruct();
+
+	void BeginRange(FUtf8StringView Id);
+	void BeginRange();
+	void EndRange();
+
+	void AddLeafId(FUtf8StringView Id);
+	template<typename T>
+	void AddLeafValue(T Value);
+	template<typename T>
+	void AddLeaf(FUtf8StringView Id, T Value) { AddLeafId(Id); AddLeafValue(Value); }
+	template<typename T>
+	void AddLeaf(T Value);
+
+	void AddComment(FUtf8StringView Comment);
+
+private:
+	void AppendNewLine();
+	void AppendIndentation();
+	void AppendIdentifier(FUtf8StringView Id);
+
+	struct FStackInfo
+	{
+		bool IsEmpty = true;
+		bool IsInStruct = true;
+	};
+
+	FUtf8StringBuilderBase& Text;
+	TArray<FStackInfo, TInlineAllocator<32>> Stack;
+	uint32 IndentLevel = 0;
+	bool IsNewLine = true;
+};
+
+///////////////////////////////////////////////////////////////////////////////
+
+class FMemberPrinter
+{
+public:
+	FMemberPrinter(FYamlBuilder& InTextBuilder, const FDebugIds& InNames)
+	: TextBuilder(InTextBuilder)
+	, Names(InNames)
+	{}
+
+	void PrintMembers(const FStructView& StructView);
+
+private:
+	void PrintLeaf(FMemberId Id, const FLeafView& LeafView);
+	void PrintStruct(FOptionalMemberId Id, FStructType StructType, const FStructView& StructView);
+	void PrintRange(FMemberId Id, ERangeSizeType NumType, const FRangeView& RangeView);
+
+	void PrintLeaves(FUnpackedLeafType Leaf, const FLeafRangeView& LeafRange);
+	void PrintStructs(FStructType StructType, const FStructRangeView& StructRange);
+	void PrintRanges(ERangeSizeType NumType, const FNestedRangeView& NestedRange);
+
+	void PrintMembersInternal(FStructType StructType, const FStructView& StructView, const FStructSchema& Schema);
+	void PrintRangeInternal(ERangeSizeType NumType, const FRangeView& RangeView);
+
+	bool IsUnicodeString(const FRangeView& RangeView);
+	void PrintUnicodeLeafValue(const FLeafView& LeafView);
+	void PrintUnicodeRangeAsLeafValue(ERangeSizeType NumType, const FRangeView& RangeView);
+
+	FYamlBuilder& TextBuilder;
+	const FDebugIds& Names;
+};
+
+///////////////////////////////////////////////////////////////////////////////
+
+void PrintBatchAsYamlDoc(
+	FUtf8StringBuilderBase& Out,
+	const FDebugIds& Names,
+	FReadBatchId Batch,
+	TConstArrayView<FStructView> Objects)
+{
+	FYamlBuilder YamlBuilder(Out);
+	FBatchPrinter Printer(YamlBuilder, Names);
+	YamlBuilder.BeginDocument();
+	Printer.PrintSchemas(Batch);
+	Printer.PrintObjects(Objects);
+	YamlBuilder.EndDocument();
+}
+
+///////////////////////////////////////////////////////////////////////////////
+
+FStructSchemaReader::FStructSchemaReader(const FStructSchema& Schema, FReadBatchId InBatch)
+: Footer(Schema.Footer)
+, Batch(InBatch)
+, TypeId(Schema.Type)
+, bIsDense(Schema.IsDense)
+, bHasSuper(Schema.Inheritance != ESuper::No)
+, bUsesSuper(UsesSuper(Schema.Inheritance))
+, NumMembers(Schema.NumMembers)
+, NumRangeTypes(Schema.NumRangeTypes)
+, NumInnerSchemas(Schema.NumInnerSchemas)
+, InnerSchemaIdx(SkipDeclaredSuperSchema(Schema.Inheritance))
+{
+	check(InnerSchemaIdx <= NumInnerSchemas);
+	checkf(NumRangeTypes != 0xFFFFu, TEXT("GrabRangeTypes() doesn't check for wrap-around"));
+}
+
+FOptionalMemberId FStructSchemaReader::PeekName() const
+{
+	int32 MemberNameIdx = MemberIdx - bUsesSuper;
+	return MemberNameIdx >= 0 ? ToOptional(GetMemberNames()[MemberNameIdx]) : NoId;
+}
+
+EMemberKind FStructSchemaReader::PeekKind() const
+{
+	return PeekType().GetKind();
+}
+
+FMemberType	FStructSchemaReader::PeekType() const
+{
+	check(HasMore());
+	return GetMemberTypes()[MemberIdx];
+}
+
+FStructSchemaHandle FStructSchemaReader::GetSuper() const
+{
+	check(HasSuper());
+	check(NumInnerSchemas > 0);
+	return { static_cast<FStructSchemaId>(GetInnerSchemas()[0]), Batch };
+}
+
+FMemberSchemaView FStructSchemaReader::GrabMember()
+{
+	check(HasMore());
+	FMemberType Type = PeekType();
+	FMemberSchemaView Out{ Type, Batch };
+
+	switch (PeekKind())
+	{
+		case EMemberKind::Leaf:
+			Out.InnerSchema = GrabLeafSchema(Type.AsLeaf());
+			break;
+		case EMemberKind::Struct:
+			Out.InnerSchema = GrabStructSchema(Type.AsStruct());
+			break;
+		case EMemberKind::Range:
+			Out.InnerRangeTypes = GrabRangeTypes();
+			Out.InnerSchema = GrabRangeSchema(Out.InnerRangeTypes.Last());
+			break;
+	}
+
+	AdvanceToNextMember();
+
+	return Out;
+}
+
+FStructSchemaReader::FMemberTypeRange FStructSchemaReader::GrabRangeTypes()
+{
+	return GrabInnerRangeTypes(MakeArrayView(GetRangeTypes(), NumRangeTypes), /* in-out */ RangeTypeIdx);
+}
+
+FSchemaId FStructSchemaReader::GrabInnerSchema()
+{
+	check(InnerSchemaIdx < NumInnerSchemas);
+	uint32 Idx = InnerSchemaIdx++;
+	return GetInnerSchemas()[Idx];
+}
+
+FOptionalSchemaId FStructSchemaReader::GrabLeafSchema(FLeafType Leaf)
+{
+	return Leaf.Type == ELeafType::Enum ? ToOptional(GrabInnerSchema()) : NoId;
+}
+
+FOptionalSchemaId FStructSchemaReader::GrabStructSchema(FStructType Struct)
+{
+	return Struct.IsDynamic ? NoId : ToOptional(GrabInnerSchema());
+}
+
+FOptionalSchemaId FStructSchemaReader::GrabRangeSchema(FMemberType InnermostType)
+{
+	check(!InnermostType.IsRange());
+	return InnermostType.IsStruct() ?
+		GrabStructSchema(InnermostType.AsStruct()) :
+		GrabLeafSchema(InnermostType.AsLeaf());
+}
+
+const FMemberType* FStructSchemaReader::GetMemberTypes() const
+{
+	return FStructSchema::GetMemberTypes(Footer);
+}
+
+const FMemberType* FStructSchemaReader::GetRangeTypes() const
+{
+	return FStructSchema::GetRangeTypes(Footer, NumMembers);
+}
+
+const FSchemaId* FStructSchemaReader::GetInnerSchemas() const
+{
+	return FStructSchema::GetInnerSchemas(Footer, NumMembers, NumRangeTypes, NumMembers - bUsesSuper);
+}
+
+const FMemberId* FStructSchemaReader::GetMemberNames() const
+{
+	return FStructSchema::GetMemberNames(Footer, NumMembers, NumRangeTypes);
+}
+
+///////////////////////////////////////////////////////////////////////////////
+
+void FYamlBuilderDeleter::operator()(FYamlBuilder* YamlBuilder) const
+{
+	delete YamlBuilder;
+}
+
+FYamlBuilderPtr MakeYamlBuilder(FUtf8StringBuilderBase& StringBuilder)
+{
+	return FYamlBuilderPtr(new FYamlBuilder(StringBuilder));
+}
+
+///////////////////////////////////////////////////////////////////////////////
+
+FYamlBuilder::FYamlBuilder(FUtf8StringBuilderBase& InStringBuilder)
+: Text(InStringBuilder)
+{
+	Stack.Emplace();
+}
+
+FYamlBuilder::~FYamlBuilder()
+{
+	Stack.Pop(EAllowShrinking::No);
+	check(Stack.IsEmpty());
+}
+
+void FYamlBuilder::BeginDocument()
+{
+	Text << "---";
+	IsNewLine = false;
+	AppendNewLine();
+	Stack.Emplace();
+}
+
+void FYamlBuilder::EndDocument()
+{
+	AppendNewLine();
+	Text << "...";
+	Stack.Pop(EAllowShrinking::No);
+}
+
+void FYamlBuilder::BeginStruct(FUtf8StringView Id)
+{
+	AppendNewLine();
+	AppendIndentation();
+	AppendIdentifier(Id);
+	Stack.Last().IsEmpty = false;
+	Stack.Emplace();
+	++IndentLevel;
+}
+
+void FYamlBuilder::BeginStruct()
+{
+	AppendNewLine();
+	AppendIndentation();
+	Stack.Last().IsEmpty = false;
+	Stack.Emplace();
+	++IndentLevel;
+}
+
+void FYamlBuilder::EndStruct()
+{
+	--IndentLevel;
+	if (Stack.Last().IsEmpty)
+	{
+		Text << " {}";
+		IsNewLine = false;
+	}
+	Stack.Pop(EAllowShrinking::No);
+}
+
+void FYamlBuilder::BeginRange(FUtf8StringView Id)
+{
+	AppendNewLine();
+	AppendIndentation();
+	AppendIdentifier(Id);
+	Stack.Last().IsEmpty = false;
+	Stack.Emplace_GetRef().IsInStruct = false;
+	++IndentLevel;
+}
+
+void FYamlBuilder::BeginRange()
+{
+	AppendNewLine();
+	AppendIndentation();
+	Stack.Last().IsEmpty = false;
+	Stack.Emplace_GetRef().IsInStruct = false;
+	++IndentLevel;
+}
+
+void FYamlBuilder::EndRange()
+{
+	if (Stack.Last().IsEmpty)
+	{
+		Text << " []";
+		IsNewLine = false;
+	}
+	Stack.Pop(EAllowShrinking::No);
+	--IndentLevel;
+}
+
+void FYamlBuilder::AddLeafId(FUtf8StringView Id)
+{
+	AppendNewLine();
+	AppendIndentation();
+	AppendIdentifier(Id);
+	IsNewLine = false;
+	Stack.Last().IsEmpty = false;
+}
+
+template<typename T>
+void FYamlBuilder::AddLeafValue(T Value)
+{
+	Text << " '" << Value << '\'';
+	IsNewLine = false;
+	Stack.Last().IsEmpty = false;
+}
+
+template<typename T>
+void FYamlBuilder::AddLeaf(T Value)
+{
+	AppendNewLine();
+	AppendIndentation();
+	Text << '\'' << Value << '\'';
+	IsNewLine = false;
+	Stack.Last().IsEmpty = false;
+}
+
+void FYamlBuilder::AddComment(FUtf8StringView Comment)
+{
+	Text << " #" << Comment;
+	AppendNewLine();
+}
+
+void FYamlBuilder::AppendNewLine()
+{
+	if (!IsNewLine)
+	{
+		Text << '\n';
+		IsNewLine = true;
+	}
+}
+
+void FYamlBuilder::AppendIndentation()
+{
+	for (uint32 I = 0; I < IndentLevel; ++I)
+	{
+		Text << "  ";
+	}
+	if (!Stack.Last().IsInStruct)
+	{
+		Text << "- ";
+	}
+	IsNewLine = false;
+}
+
+void FYamlBuilder::AppendIdentifier(FUtf8StringView Id)
+{
+	Text << '\'' << Id << "':";
+	IsNewLine = false;
+}
+
+///////////////////////////////////////////////////////////////////////////////
+
+
+FBatchPrinter::FBatchPrinter(FYamlBuilder& InTextBuilder, const FDebugIds& InNames)
+: TextBuilder(InTextBuilder)
+, Names(InNames)
+{
+}
+
+FBatchPrinter::~FBatchPrinter()
+{}
+
+void FBatchPrinter::PrintSchemas(FReadBatchId Batch)
+{
+	PrintSchemas(GetReadSchemas(Batch), Batch);
+}
+
+void FBatchPrinter::PrintSchemas(const FSchemaBatch& Batch, FReadBatchId BatchId)
+{
+	TextBuilder.BeginRange("Structs");
+	for (const FStructSchema& Struct : GetStructSchemas(Batch))
+	{
+		PrintStructSchema(Struct, BatchId);
+	}
+	TextBuilder.EndRange();
+
+	TextBuilder.BeginRange("Enums");
+	for (const FEnumSchema& EnumSchema : GetEnumSchemas(Batch))
+	{
+		PrintEnumSchema(EnumSchema);
+	}
+	TextBuilder.EndRange();
+}
+
+void FBatchPrinter::PrintObjects(TConstArrayView<FStructView> Objects)
+{
+	TextBuilder.BeginRange("Objects");
+	for (FStructView Object : Objects)
+	{
+		FMemberPrinter(TextBuilder, Names).PrintMembers(Object);
+	}
+	TextBuilder.EndRange();
+}
+
+static FUtf8StringView PrintMemberSchema(
+	FUtf8StringBuilderBase& Out,
+	const FDebugIds& Names,
+	const FMemberSchemaView& Schema)
+{
+	return PrintMemberSchemaImpl(Out, Names, Schema.Type, Schema.InnerSchema, Schema.InnerRangeTypes);
+}
+
+template <int32 BufferSize>
+class FPrintId
+{
+	TUtf8StringBuilder<BufferSize>			Buffer;
+public:
+	template <typename T>
+	FPrintId(const FDebugIds& Names, T Id)	{ Names.AppendString(Buffer, Id); }
+	FUtf8StringView operator*() const		{ return Buffer.ToView(); }
+};
+
+void FBatchPrinter::PrintStructSchema(const FStructSchema& Struct, FReadBatchId BatchId)
+{
+	FStructSchemaReader Reader(Struct, BatchId);
+
+	TextBuilder.BeginStruct(*FPrintId<128>(Names, Reader.GetTypeId()));
+
+	if (Reader.HasSuper())
+	{
+		const FStructSchema& SuperSchema = Reader.GetSuper().Resolve();
+		TextBuilder.AddLeaf("Super", *FPrintId<128>(Names, SuperSchema.Type));
+	}
+
+	TextBuilder.BeginRange("Members");
+	while (Reader.HasMore())
+	{
+		TUtf8StringBuilder<256> Buf;
+		FOptionalMemberId MemberId = Reader.PeekName();
+		TextBuilder.AddLeaf(*FPrintId<128>(Names, MemberId), PrintMemberSchema(Buf, Names, Reader.GrabMember()));
+	}
+	TextBuilder.EndRange();
+
+	TextBuilder.EndStruct();
+}
+
+void FBatchPrinter::PrintEnumSchema(const FEnumSchema& Enum)
+{
+	TextBuilder.BeginStruct(*FPrintId<128>(Names, Enum.Type));
+	TextBuilder.AddLeaf("FlagMode", !!Enum.FlagMode);
+
+	TextBuilder.BeginRange("Constants");
+	TConstArrayView<FNameId> EnumNames = MakeConstArrayView(Enum.Footer, Enum.Num);
+	switch (Enum.Width)
+	{
+		case ELeafWidth::B8:
+			PrintEnumConstants(EnumNames, GetConstants<uint8>(Enum), Enum.FlagMode);
+			break;
+		case ELeafWidth::B16:
+			PrintEnumConstants(EnumNames, GetConstants<uint16>(Enum), Enum.FlagMode);
+			break;
+		case ELeafWidth::B32:
+			PrintEnumConstants(EnumNames, GetConstants<uint32>(Enum), Enum.FlagMode);
+			break;
+		case ELeafWidth::B64:
+			PrintEnumConstants(EnumNames, GetConstants<uint64>(Enum), Enum.FlagMode);
+			break;
+	}
+	TextBuilder.EndRange();
+
+	TextBuilder.EndStruct();
+}
+
+template<typename IntType>
+void FBatchPrinter::PrintEnumConstants(
+	TConstArrayView<FNameId> EnumNames,
+	TConstArrayView<IntType> Constants,
+	bool bFlagMode)
+{
+	uint16 NamesNum = IntCastChecked<uint16>(EnumNames.Num());
+	if (Constants.Num() > 0)
+	{
+		check(EnumNames.Num() == Constants.Num());
+		for (uint16 Idx = 0; Idx < NamesNum; ++Idx)
+		{
+			TextBuilder.AddLeaf(*FPrintId<128>(Names, EnumNames[Idx]), (uint64)Constants[Idx]);
+		}
+	}
+	else if (bFlagMode)
+	{
+		uint64 Value = 1;
+		for (uint16 Idx = 0; Idx < NamesNum; ++Idx)
+		{
+			TextBuilder.AddLeaf(*FPrintId<128>(Names, EnumNames[Idx]), Value);
+			Value <<= 1;
+		}
+	}
+	else
+	{
+		for (uint16 Idx = 0; Idx < NamesNum; ++Idx)
+		{
+			TextBuilder.AddLeaf(*FPrintId<128>(Names, EnumNames[Idx]), (uint64)Idx);
+		}
+	}
+}
+
+///////////////////////////////////////////////////////////////////////////////
+
+void FMemberPrinter::PrintMembers(const FStructView& StructView)
+{
+	const FStructSchema& Schema = StructView.Schema.Resolve();
+	TextBuilder.BeginStruct(*FPrintId<128>(Names, Schema.Type));
+	PrintMembersInternal({ EMemberKind::Struct }, StructView, Schema);
+}
+
+void FMemberPrinter::PrintLeaf(FMemberId Id, const FLeafView& LeafView)
+{
+	TextBuilder.AddLeafId(*FPrintId<128>(Names, Id));
+
+	switch (LeafView.Leaf.Type)
+	{
+	case ELeafType::Bool:
+		TextBuilder.AddLeafValue(LeafView.AsBool());
+		break;
+	case ELeafType::IntS:
+		switch (LeafView.Leaf.Width)
+		{
+			case ELeafWidth::B8:	TextBuilder.AddLeafValue(LeafView.AsS8()); break;
+			case ELeafWidth::B16:	TextBuilder.AddLeafValue(LeafView.AsS16()); break;
+			case ELeafWidth::B32:	TextBuilder.AddLeafValue(LeafView.AsS32()); break;
+			case ELeafWidth::B64:	TextBuilder.AddLeafValue(LeafView.AsS64()); break;
+		}
+		break;
+	case ELeafType::IntU:
+		switch (LeafView.Leaf.Width)
+		{
+			case ELeafWidth::B8:	TextBuilder.AddLeafValue(LeafView.AsU8()); break;
+			case ELeafWidth::B16:	TextBuilder.AddLeafValue(LeafView.AsU16()); break;
+			case ELeafWidth::B32:	TextBuilder.AddLeafValue(LeafView.AsU32()); break;
+			case ELeafWidth::B64:	TextBuilder.AddLeafValue(LeafView.AsU64()); break;
+		}
+		break;
+	case ELeafType::Float:
+		if (LeafView.Leaf.Width == ELeafWidth::B32)
+		{
+			TextBuilder.AddLeafValue(LeafView.AsFloat());
+		}
+		else
+		{
+			check(LeafView.Leaf.Width == ELeafWidth::B64);
+			TextBuilder.AddLeafValue(LeafView.AsDouble());
+		}
+		break;
+	case ELeafType::Hex:
+		check(LeafView.Leaf.Type != ELeafType::Hex);
+		break;
+	case ELeafType::Enum:
+		switch (LeafView.Leaf.Width)
+		{
+			case ELeafWidth::B8:	TextBuilder.AddLeafValue(LeafView.AsUnderlyingValue<uint8>()); break;
+			case ELeafWidth::B16:	TextBuilder.AddLeafValue(LeafView.AsUnderlyingValue<uint16>()); break;
+			case ELeafWidth::B32:	TextBuilder.AddLeafValue(LeafView.AsUnderlyingValue<uint32>()); break;
+			case ELeafWidth::B64:	TextBuilder.AddLeafValue(LeafView.AsUnderlyingValue<uint64>()); break;
+		}
+		break;
+	case ELeafType::Unicode:
+		PrintUnicodeLeafValue(LeafView);
+		break;
+	}
+
+	if constexpr(PrintWithComments)
+	{
+		TUtf8StringBuilder<256> Buf;
+		PrintMemberSchemaImpl(Buf, Names, FMemberType(LeafView.Leaf.Pack()), static_cast<FOptionalSchemaId>(LeafView.Enum));
+		TextBuilder.AddComment(Buf);
+	}
+}
+
+void FMemberPrinter::PrintStruct(FOptionalMemberId MemberId, FStructType StructType, const FStructView& StructView)
+{
+	TextBuilder.BeginStruct(*FPrintId<128>(Names, MemberId));
+	PrintMembersInternal(StructType, StructView, StructView.Schema.Resolve());
+}
+
+void FMemberPrinter::PrintRange(FMemberId Id, ERangeSizeType NumType, const FRangeView& RangeView)
+{
+	if (IsUnicodeString(RangeView))
+	{
+		TextBuilder.AddLeafId(*FPrintId<128>(Names, Id));
+		PrintUnicodeRangeAsLeafValue(NumType, RangeView);
+	}
+	else
+	{
+		TextBuilder.BeginRange(*FPrintId<128>(Names, Id));
+		PrintRangeInternal(NumType, RangeView);
+	}
+}
+
+void FMemberPrinter::PrintLeaves(FUnpackedLeafType Leaf, const FLeafRangeView& LeafRange)
+{
+	switch (Leaf.Type)
+	{
+	case ELeafType::Bool:
+		for (bool b : LeafRange.AsBools()) { TextBuilder.AddLeaf(b); } break;
+	case ELeafType::IntS:
+		switch (Leaf.Width)
+		{
+			case ELeafWidth::B8:  for (const int8&  I : LeafRange.AsS8s())  { TextBuilder.AddLeaf(I); } break;
+			case ELeafWidth::B16: for (const int16& I : LeafRange.AsS16s()) { TextBuilder.AddLeaf(I); } break;
+			case ELeafWidth::B32: for (const int32& I : LeafRange.AsS32s()) { TextBuilder.AddLeaf(I); } break;
+			case ELeafWidth::B64: for (const int64& I : LeafRange.AsS64s()) { TextBuilder.AddLeaf(I); } break;
+		}
+		break;
+	case ELeafType::IntU:
+		switch (Leaf.Width)
+		{
+			case ELeafWidth::B8:  for (const uint8&  U : LeafRange.AsU8s())  { TextBuilder.AddLeaf(U); } break;
+			case ELeafWidth::B16: for (const uint16& U : LeafRange.AsU16s()) { TextBuilder.AddLeaf(U); } break;
+			case ELeafWidth::B32: for (const uint32& U : LeafRange.AsU32s()) { TextBuilder.AddLeaf(U); } break;
+			case ELeafWidth::B64: for (const uint64& U : LeafRange.AsU64s()) { TextBuilder.AddLeaf(U); } break;
+		}
+		break;
+	case ELeafType::Float:
+		if (Leaf.Width == ELeafWidth::B32)
+		{
+			for (const float& f : LeafRange.AsFloats()) { TextBuilder.AddLeaf(f); }
+		}
+		else
+		{
+			check(Leaf.Width == ELeafWidth::B64);
+			for (const double& d : LeafRange.AsDoubles()) { TextBuilder.AddLeaf(d); }
+		}
+		break;
+	case ELeafType::Hex:
+		// TODO: Implement AddLeaf(Hex)
+		check(Leaf.Type != ELeafType::Hex);
+		break;
+	case ELeafType::Enum:
+		switch (Leaf.Width)
+		{
+			case ELeafWidth::B8:  for (const uint8&  U : LeafRange.AsUnderlyingValues<uint8>())  { TextBuilder.AddLeaf(U); } break;
+			case ELeafWidth::B16: for (const uint16& U : LeafRange.AsUnderlyingValues<uint16>()) { TextBuilder.AddLeaf(U); } break;
+			case ELeafWidth::B32: for (const uint32& U : LeafRange.AsUnderlyingValues<uint32>()) { TextBuilder.AddLeaf(U); } break;
+			case ELeafWidth::B64: for (const uint64& U : LeafRange.AsUnderlyingValues<uint64>()) { TextBuilder.AddLeaf(U); } break;
+		}
+		break;
+	case ELeafType::Unicode:
+		checkf(LeafRange.Num() == 0, TEXT("Should have been handled by PrintUnicodeRangeAsLeafValue"));
+		break;
+	}
+}
+
+void FMemberPrinter::PrintStructs(FStructType StructType, const FStructRangeView& StructRange)
+{
+	for (FStructView StructView : StructRange)
+	{
+		TextBuilder.BeginStruct();
+		PrintMembersInternal(StructType, StructView, StructView.Schema.Resolve());
+	}
+}
+
+void FMemberPrinter::PrintRanges(ERangeSizeType NumType, const FNestedRangeView& NestedRange)
+{
+	for (FRangeView RangeView : NestedRange)
+	{
+		if (IsUnicodeString(RangeView))
+		{
+			TextBuilder.AddLeafId({});
+			PrintUnicodeRangeAsLeafValue(NumType, RangeView);
+		}
+		else
+		{
+			TextBuilder.BeginRange();
+			PrintRangeInternal(NumType, RangeView);
+		}
+	}
+}
+
+void FMemberPrinter::PrintMembersInternal(FStructType StructType, const FStructView& StructView, const FStructSchema& Schema)
+{
+	FMemberReader It(Schema, StructView.Values, StructView.Schema.Batch);
+
+	const bool HasMembers = StructType.IsDynamic || It.HasMore();
+	if constexpr(PrintWithComments)
+	{
+		if (HasMembers)
+		{
+			TUtf8StringBuilder<256> Buf;
+			TextBuilder.AddComment(PrintMemberSchemaImpl(Buf, Names, FMemberType(StructType), StructView.Schema.Id));
+		}
+	}
+
+	if (StructType.IsDynamic)
+	{
+		TextBuilder.AddLeaf(GLiterals.Dynamic, *FPrintId<128>(Names, Schema.Type));
+	}
+	while (It.HasMore())
+	{
+		FOptionalMemberId Id = It.PeekName();
+		FMemberType Type = It.PeekType();
+		switch (Type.GetKind())
+		{
+			case EMemberKind::Leaf:
+				PrintLeaf(Id.Get(), It.GrabLeaf());
+				break;
+			case EMemberKind::Struct:
+				PrintStruct(Id, Type.AsStruct(), It.GrabStruct());
+				break;
+			case EMemberKind::Range:
+				PrintRange(Id.Get(), Type.AsRange().MaxSize, It.GrabRange());
+				break;
+		}
+	}
+	TextBuilder.EndStruct();
+
+	if constexpr(PrintWithComments)
+	{
+		if (!HasMembers)
+		{
+			TUtf8StringBuilder<256> Buf;
+			TextBuilder.AddComment(PrintMemberSchemaImpl(Buf, Names, FMemberType(StructType), StructView.Schema.Id));
+		}
+	}
+}
+
+static TArray<FMemberType, TInlineAllocator<16>> GetInnerRangeTypes(const FRangeSchema& Schema)
+{
+	TArray<FMemberType, TInlineAllocator<16>> InnerRangeTypes;
+	InnerRangeTypes.Add(Schema.ItemType);
+	if (Schema.ItemType.IsRange())
+	{
+		const FMemberType* It = Schema.NestedItemTypes;
+		do { InnerRangeTypes.Add(*It); } while ((It++)->IsRange());
+	}
+	return InnerRangeTypes;
+}
+
+static void AddRangeSchemaComment(
+	FYamlBuilder& TextBuilder,
+	const FDebugIds& Names,
+	FMemberType Type,
+	const FRangeSchema& Schema)
+{
+	TUtf8StringBuilder<256> Comment;
+	TextBuilder.AddComment(PrintMemberSchemaImpl(Comment, Names, Type, Schema.InnermostSchema, GetInnerRangeTypes(Schema)));
+}
+
+void FMemberPrinter::PrintRangeInternal(ERangeSizeType NumType, const FRangeView& RangeView)
+{
+	const FRangeSchema& Schema = RangeView.Schema;
+	if constexpr(PrintWithComments)
+	{
+		if (RangeView.Num() > 0)
+		{
+			AddRangeSchemaComment(TextBuilder, Names, FMemberType(NumType), Schema);
+		}
+	}
+
+	switch (Schema.ItemType.GetKind())
+	{
+		case EMemberKind::Leaf:
+			PrintLeaves(Schema.ItemType.AsLeaf(), RangeView.AsLeaves());
+			break;
+		case EMemberKind::Struct:
+			PrintStructs(Schema.ItemType.AsStruct(), RangeView.AsStructs());
+			break;
+		case EMemberKind::Range:
+			PrintRanges(Schema.ItemType.AsRange().MaxSize, RangeView.AsRanges());
+			break;
+	}
+	TextBuilder.EndRange();
+
+	if constexpr(PrintWithComments)
+	{
+		if (RangeView.Num() == 0)
+		{
+			AddRangeSchemaComment(TextBuilder, Names, FMemberType(NumType), Schema);
+		}
+	}
+}
+
+bool FMemberPrinter::IsUnicodeString(const FRangeView& RangeView)
+{
+	FMemberType Type = RangeView.Schema.ItemType;
+	return RangeView.Num() > 0 && Type.IsLeaf() && Type.AsLeaf().Type == ELeafType::Unicode;
+}
+
+void FMemberPrinter::PrintUnicodeLeafValue(const FLeafView& LeafView)
+{
+	if (LeafView.Leaf.Width == ELeafWidth::B8)
+	{
+		UTF8CHAR Char8 = static_cast<UTF8CHAR>(LeafView.AsChar8());
+		check(Char8 <= 127);
+		if (Char8 == 0)
+		{
+			TextBuilder.AddLeafValue(FUtf8StringView{});
+		}
+		else
+		{
+			TextBuilder.AddLeafValue(Char8);
+		}
+	}
+	else if (LeafView.Leaf.Width == ELeafWidth::B16)
+	{
+		UTF16CHAR Char16 = static_cast<UTF16CHAR>(LeafView.AsChar16());
+		if (Char16 == 0)
+		{
+			TextBuilder.AddLeafValue(FUtf8StringView{});
+		}
+		else if (Char16 <= 127)
+		{
+			TextBuilder.AddLeafValue(static_cast<UTF8CHAR>(Char16));
+		}
+		else
+		{
+			TextBuilder.AddLeafValue(FUtf8StringView(StringCast<UTF8CHAR>(&Char16, 1)));
+		}
+	}
+	else // ELeafWidth::B32
+	{
+		check(LeafView.Leaf.Width != ELeafWidth::B64);
+		UTF32CHAR Char32 = static_cast<UTF32CHAR>(LeafView.AsChar32());
+		if (Char32 == 0)
+		{
+			TextBuilder.AddLeafValue(FUtf8StringView{});
+		}
+		else if (Char32 <= 127)
+		{
+			TextBuilder.AddLeafValue(static_cast<UTF8CHAR>(Char32));
+		}
+		else
+		{
+			TextBuilder.AddLeafValue(FUtf8StringView(StringCast<UTF8CHAR>(&Char32, 1)));
+		}
+	}
+}
+
+void FMemberPrinter::PrintUnicodeRangeAsLeafValue(ERangeSizeType NumType, const FRangeView& RangeView)
+{
+	check(IsUnicodeString(RangeView));
+
+	const FLeafRangeView LeafRange = RangeView.AsLeaves();
+	FUnpackedLeafType Leaf = RangeView.Schema.ItemType.AsLeaf();
+
+	if (Leaf.Width == ELeafWidth::B8)
+	{
+		TRangeView<char8_t> Utf8 = LeafRange.AsUtf8();
+		check(Utf8.Num() > 0);
+		const UTF8CHAR* Src = reinterpret_cast<const UTF8CHAR*>(Utf8.begin());
+		int32 Len = IntCastChecked<int32>(Utf8.Num());
+		// explicitly strip null terminator, no StringCast required
+		Len -= (Utf8[Len-1] == UTF8CHAR(0));
+		TextBuilder.AddLeafValue(FUtf8StringView(Src, Len));
+	}
+	else if (Leaf.Width == ELeafWidth::B16)
+	{
+		TRangeView<char16_t> Utf16 = LeafRange.AsUtf16();
+		check(Utf16.Num() > 0);
+		const UTF16CHAR* Src = reinterpret_cast<const UTF16CHAR*>(Utf16.begin());
+		const int32 Len = IntCastChecked<int32>(Utf16.Num());
+		// null terminator is excluded by length of StringCast result
+		TextBuilder.AddLeafValue(FUtf8StringView(StringCast<UTF8CHAR>(Src, Len)));
+	}
+	else // ELeafWidth::B32
+	{
+		check(Leaf.Width != ELeafWidth::B64);
+		TRangeView<char32_t> Utf32 = LeafRange.AsUtf32();
+		check(Utf32.Num() > 0);
+		const UTF32CHAR* Src = reinterpret_cast<const UTF32CHAR*>(Utf32.begin());
+		const int32 Len = IntCastChecked<int32>(Utf32.Num());
+		// null terminator is excluded by length of StringCast result
+		TextBuilder.AddLeafValue(FUtf8StringView(StringCast<UTF8CHAR>(Src, Len)));
+	}
+
+	if constexpr(PrintWithComments)
+	{
+		AddRangeSchemaComment(TextBuilder, Names, FMemberType(NumType), RangeView.Schema);
+	}
+}
+
+///////////////////////////////////////////////////////////////////////////////
+
+void FDebugIds::AppendString(FUtf8StringBuilderBase& Out, FMemberId Name) const
+{
+	AppendString(Out, Name.Id);
+}
+
+void FDebugIds::AppendString(FUtf8StringBuilderBase& Out, FOptionalMemberId Name) const
+{
+	if (Name)
+	{
+		AppendString(Out, Name.Get().Id);
+	}
+	else
+	{
+		Out.Append(GLiterals.Super);
+	}
+}
+
+void FDebugIds::AppendString(FUtf8StringBuilderBase& Out, FScopeId Scope) const
+{
+	if (Scope.IsFlat())
+	{
+		AppendString(Out, Scope.AsFlat().Name);
+	}
+	else if (Scope)
+	{
+		FNestedScope Nested = Resolve(Scope.AsNested());
+		AppendString(Out, Nested.Outer);
+		Out.AppendChar('.');
+		AppendString(Out, Nested.Inner.Name);
+	}
+}
+
+void FDebugIds::AppendString(FUtf8StringBuilderBase& Out, FTypenameId Typename) const
+{
+	if (Typename.IsConcrete())
+	{
+		AppendString(Out, Typename.AsConcrete().Id);
+	}
+	else
+	{
+		FParametricTypeView ParametricType = Resolve(Typename.AsParametric());
+		TConstArrayView<FTypeId> Parameters = ParametricType.GetParameters();
+	
+		if (ParametricType.Name)
+		{
+			AppendString(Out, ParametricType.Name.Get().Id);
+		}
+
+		Out.AppendChar(ParametricType.Name ? '<' : '[');
+		for (FTypeId Parameter : Parameters.LeftChop(1))
+		{
+			AppendString(Out, Parameter);
+			Out.AppendChar(',');
+		}
+		if (Parameters.Num() > 0)
+		{
+			AppendString(Out, Parameters.Last());
+		}
+		Out.AppendChar(ParametricType.Name ? '>' : ']');
+	}
+}
+
+void FDebugIds::AppendString(FUtf8StringBuilderBase& Out, FTypeId Type) const
+{
+	if (Type.Scope)
+	{
+		AppendString(Out, Type.Scope);
+		Out.AppendChar('.');
+	}
+	AppendString(Out, Type.Name);
+}
+
+void FDebugIds::AppendString(FUtf8StringBuilderBase& Out, FEnumSchemaId Name) const
+{
+	AppendString(Out, Resolve(Name));
+}
+
+void FDebugIds::AppendString(FUtf8StringBuilderBase& Out, FStructSchemaId Name) const
+{
+	AppendString(Out, Resolve(Name));
+}
+
+FString FDebugIds::Print(FNameId Name) const
+{
+	TUtf8StringBuilder<128> Out;
+	AppendString(Out, Name);
+	return FString(FStringView(StringCast<TCHAR>(Out.GetData(), Out.Len())));
+}
+
+FString FDebugIds::Print(FMemberId Name) const
+{
+	TUtf8StringBuilder<128> Out;
+	AppendString(Out, Name.Id);
+	return FString(FStringView(StringCast<TCHAR>(Out.GetData(), Out.Len())));
+}
+
+FString FDebugIds::Print(FOptionalMemberId Name) const
+{
+	TUtf8StringBuilder<128> Out;
+	AppendString(Out, Name);
+	return FString(FStringView(StringCast<TCHAR>(Out.GetData(), Out.Len())));
+}
+
+FString FDebugIds::Print(FTypeId Type) const
+{
+	TUtf8StringBuilder<128> Out;
+	AppendString(Out, Type);
+	return FString(FStringView(StringCast<TCHAR>(Out.GetData(), Out.Len())));
+}
+
+FString FDebugIds::Print(FEnumSchemaId Name) const
+{
+	TUtf8StringBuilder<128> Out;
+	AppendString(Out, Name);
+	return FString(FStringView(StringCast<TCHAR>(Out.GetData(), Out.Len())));
+}
+
+FString FDebugIds::Print(FStructSchemaId Name) const
+{
+	TUtf8StringBuilder<128> Out;
+	AppendString(Out, Name);
+	return FString(FStringView(StringCast<TCHAR>(Out.GetData(), Out.Len())));
+}
+
+///////////////////////////////////////////////////////////////////////////////
+
+} // namespace PlainProps
diff --git a/Engine/Plugins/Experimental/PlainProps/Source/Private/PlainPropsRead.cpp b/Engine/Plugins/Experimental/PlainProps/Source/Private/PlainPropsRead.cpp
--- a/Engine/Plugins/Experimental/PlainProps/Source/Private/PlainPropsRead.cpp
+++ b/Engine/Plugins/Experimental/PlainProps/Source/Private/PlainPropsRead.cpp
@@ -36,24 +36,6 @@ void FSchemaBatch::ValidateBounds(uint64 NumBytes) const
 	check(ExpectedEnd == ActualEnd);
 }
 
-template<class SchemaType>
-const SchemaType& ResolveSchema(const FSchemaBatch& Batch, FSchemaId Id)
-{
-	check(Id.Idx < Batch.NumSchemas);
-	return *reinterpret_cast<const SchemaType*>(reinterpret_cast<const uint8*>(&Batch) + Batch.SchemaOffsets[Id.Idx]);
-}
-	
-static FNestedScope ResolveNestedScope(const FSchemaBatch& Batch, FNestedScopeId Id)
-{
-	return Batch.GetNestedScopes()[Id.Idx];
-}
-	
-static FParametricTypeView ResolveParametricType(const FSchemaBatch& Batch, FParametricTypeId Id)
-{
-	FParametricType Type = Batch.GetParametricTypes()[Id.Idx];
-	return {Type.Name, IntCastChecked<uint8>(Type.Parameters.NumParameters), Batch.GetFirstParameter() + Type.Parameters.Idx};
-}
-
 
 class FReadSchemaRegistry
 {
@@ -120,14 +102,19 @@ uint32 NumStructSchemas(FReadBatchId Batch)
 	return GReadSchemas.Get(Batch).NumStructSchemas;
 }
 
+const FSchemaBatch& GetReadSchemas(FReadBatchId Batch)
+{
+	return GReadSchemas.Get(Batch);
+}
+
 const FStructSchema& ResolveStructSchema(FReadBatchId Batch, FStructSchemaId Schema)
 {
-	return ResolveSchema<FStructSchema>(GReadSchemas.Get(Batch), Schema);
+	return ResolveStructSchema(GReadSchemas.Get(Batch), Schema);
 }
 
 const FEnumSchema& ResolveEnumSchema(FReadBatchId Batch, FEnumSchemaId Schema)
 {
-	return ResolveSchema<FEnumSchema>(GReadSchemas.Get(Batch), Schema);
+	return ResolveEnumSchema(GReadSchemas.Get(Batch), Schema);
 }
 
 FNestedScope ResolveUntranslatedNestedScope(FReadBatchId Batch, FNestedScopeId Id)
@@ -191,11 +178,6 @@ const FMemberType*	FMemberReader::GetRangeTypes() const	{ return FStructSchema::
 const FSchemaId*	FMemberReader::GetInnerSchemas() const	{ return FStructSchema::GetInnerSchemas(Footer, NumMembers, NumRangeTypes, NumMembers - HasSuper); }
 const FMemberId*	FMemberReader::GetMemberNames() const	{ return FStructSchema::GetMemberNames(Footer, NumMembers, NumRangeTypes); }
 
-static bool SkipDeclaredSuperSchema(ESuper Inheritance)
-{
-	return Inheritance == ESuper::Unused || Inheritance == ESuper::Used;
-}
-
 FMemberReader::FMemberReader(const FStructSchema& Schema, FByteReader Values, FReadBatchId InBatch)
 : Footer(Schema.Footer)
 , Batch(InBatch)
diff --git a/Engine/Plugins/Experimental/PlainProps/Source/Private/PlainPropsUeCoreBindings.cpp b/Engine/Plugins/Experimental/PlainProps/Source/Private/PlainPropsUeCoreBindings.cpp
--- a/Engine/Plugins/Experimental/PlainProps/Source/Private/PlainPropsUeCoreBindings.cpp
+++ b/Engine/Plugins/Experimental/PlainProps/Source/Private/PlainPropsUeCoreBindings.cpp
@@ -103,9 +103,9 @@ namespace PlainProps
 {
 	
 template<>
-void AppendString(FString& Out, const FName& Name)
+void AppendString(FUtf8StringBuilderBase& Out, const FName& Name)
 {
 	Name.AppendString(Out);
 }
 
-}
\ No newline at end of file
+}
diff --git a/Engine/Plugins/Experimental/PlainProps/Source/Private/PlainPropsWrite.cpp b/Engine/Plugins/Experimental/PlainProps/Source/Private/PlainPropsWrite.cpp
--- a/Engine/Plugins/Experimental/PlainProps/Source/Private/PlainPropsWrite.cpp
+++ b/Engine/Plugins/Experimental/PlainProps/Source/Private/PlainPropsWrite.cpp
@@ -803,104 +803,4 @@ FStructSchemaId FWriter::WriteMembers(TArray64<uint8>& Out, FStructSchemaId Buil
 
 ////////////////////////////////////////////////////////////////////////////////////////////////////////////////////////////////////////////////////
 
-void FDebugIds::AppendDebugString(FString& Out, FScopeId Scope) const
-{
-	if (Scope.IsFlat())
-	{
-		AppendDebugString(Out, Scope.AsFlat().Name);
-	}
-	else if (Scope)
-	{
-		FNestedScope Nested = Resolve(Scope.AsNested());
-		AppendDebugString(Out, Nested.Outer);
-		Out.AppendChar('.');
-		AppendDebugString(Out, Nested.Inner.Name);
-	}
-}
-
-void FDebugIds::AppendDebugString(FString& Out, FTypenameId Typename) const
-{
-	if (Typename.IsConcrete())
-	{
-		AppendDebugString(Out, Typename.AsConcrete().Id);
-	}
-	else
-	{
-		FParametricTypeView ParametricType = Resolve(Typename.AsParametric());
-	
-		if (ParametricType.Name)
-		{
-			AppendDebugString(Out, ParametricType.Name.Get().Id);
-		}
-
-		Out.AppendChar(ParametricType.Name ? '<' : '[');
-		for (FTypeId Parameter : ParametricType.GetParameters())
-		{
-			AppendDebugString(Out, Parameter);
-			Out.AppendChar(',');
-		}
-		Out.GetCharArray()[Out.Len() - 1] = ParametricType.Name ? '>' : ']';
-	}
-}
-
-void FDebugIds::AppendDebugString(FString& Out, FTypeId Type) const
-{
-	if (Type.Scope)
-	{
-		AppendDebugString(Out, Type.Scope);
-		Out.AppendChar('.');
-	}
-	AppendDebugString(Out, Type.Name);
-}
-
-void FDebugIds::AppendDebugString(FString& Out, FEnumSchemaId Name) const
-{
-	AppendDebugString(Out, Resolve(Name));
-}
-
-void FDebugIds::AppendDebugString(FString& Out, FStructSchemaId Name) const
-{
-	AppendDebugString(Out, Resolve(Name));
-}
-
-FString FDebugIds::Print(FNameId Name) const
-{
-	FString Out;
-	AppendDebugString(Out, Name);
-	return Out;
-}
-
-FString FDebugIds::Print(FMemberId Name) const
-{
-	FString Out;
-	AppendDebugString(Out, Name.Id);
-	return Out;
-}
-
-FString FDebugIds::Print(FOptionalMemberId Name) const
-{
-	return Name ? Print(Name.Get().Id) : TEXT("!super!");
-}
-
-FString FDebugIds::Print(FTypeId Type) const
-{
-	FString Out;
-	AppendDebugString(Out, Type);
-	return Out;
-}
-
-FString FDebugIds::Print(FEnumSchemaId Name) const
-{
-	FString Out;
-	AppendDebugString(Out, Name);
-	return Out;
-}
-
-FString FDebugIds::Print(FStructSchemaId Name) const
-{
-	FString Out;
-	AppendDebugString(Out, Name);
-	return Out;
-}
-
 } // namespace PlainProps
diff --git a/Engine/Plugins/Experimental/PlainProps/Source/Public/PlainPropsBuild.h b/Engine/Plugins/Experimental/PlainProps/Source/Public/PlainPropsBuild.h
--- a/Engine/Plugins/Experimental/PlainProps/Source/Public/PlainPropsBuild.h
+++ b/Engine/Plugins/Experimental/PlainProps/Source/Public/PlainPropsBuild.h
@@ -16,7 +16,6 @@ namespace PlainProps
 struct FBuiltMember;
 struct FBuiltStruct;
 struct FBuiltRange;
-class FDebugIds;
 struct FUnpackedLeafType;
 
 //////////////////////////////////////////////////////////////////////////
@@ -246,7 +245,7 @@ public:
 	template<Enumeration T>
 	void AddEnum(FMemberId Name, FEnumSchemaId Schema, T Value)
 	{
-		AddLeaf(Name, ReflectEnum<T>, ToOptional(Schema), ValueCast(Value));
+		AddLeaf(Name, ReflectEnum<T>, ToOptional(Schema), ValueCast(static_cast<__underlying_type(T)>(Value)));
 	}
 	
 	void AddEnum8(FMemberId Name, FEnumSchemaId Schema, uint8 Value)	{ AddLeaf(Name, {ELeafType::Enum, ELeafWidth::B8},  ToOptional(Schema), Value); }
@@ -345,4 +344,4 @@ private:
 	FMemberSchema Schema;
 };
 
-} // namespace PlainProps
\ No newline at end of file
+} // namespace PlainProps
diff --git a/Engine/Plugins/Experimental/PlainProps/Source/Public/PlainPropsIndex.h b/Engine/Plugins/Experimental/PlainProps/Source/Public/PlainPropsIndex.h
--- a/Engine/Plugins/Experimental/PlainProps/Source/Public/PlainPropsIndex.h
+++ b/Engine/Plugins/Experimental/PlainProps/Source/Public/PlainPropsIndex.h
@@ -85,10 +85,7 @@ protected:
 };
 
 template<class NameType>
-void AppendString(FString& Out, const NameType& Str)
-{
-	Out.Append(Str);
-}
+void AppendString(FUtf8StringBuilderBase& Out, const NameType& Str);
 
 template<class NameT>
 class TIdIndexer : public FIdIndexerBase
@@ -111,9 +108,9 @@ public:
 protected:
 	TSet<NameT>									Names;
 
-	virtual void								AppendDebugString(FString& Out, FNameId Id) const override final
+	virtual void								AppendString(FUtf8StringBuilderBase& Out, FNameId Id) const override final
 	{
-		AppendString(Out, ResolveName(Id));
+		PlainProps::AppendString(Out, ResolveName(Id));
 	}
 
 	virtual FConcreteTypenameId					InitParameterName(const char* Name) override final
diff --git a/Engine/Plugins/Experimental/PlainProps/Source/Public/PlainPropsPrint.h b/Engine/Plugins/Experimental/PlainProps/Source/Public/PlainPropsPrint.h
--- a/Engine/Plugins/Experimental/PlainProps/Source/Public/PlainPropsPrint.h
+++ b/Engine/Plugins/Experimental/PlainProps/Source/Public/PlainPropsPrint.h
@@ -8,7 +8,29 @@
 namespace PlainProps
 {
 
+class FReadBatchId;
+struct FStructView;
+
+///////////////////////////////////////////////////////////////////////////////
+
 PLAINPROPS_API FAnsiStringView ToString(ERangeSizeType Width);
 PLAINPROPS_API FAnsiStringView ToString(FUnpackedLeafType Leaf);
 
+///////////////////////////////////////////////////////////////////////////////
+
+PLAINPROPS_API void PrintMemberSchema(
+	FUtf8StringBuilderBase& Out,
+	const FDebugIds& Names, 
+	FMemberType Type,
+	FOptionalSchemaId InnerSchema,
+	TConstArrayView<FMemberType> InnerRangeTypes);
+
+void PrintBatchAsYamlDoc(
+	FUtf8StringBuilderBase& Out,
+	const FDebugIds& Names,
+	FReadBatchId Batch,
+	TConstArrayView<FStructView> Objects);
+
+///////////////////////////////////////////////////////////////////////////////
+
 } // namespace PlainProps
diff --git a/Engine/Plugins/Experimental/PlainProps/Source/Public/PlainPropsRead.h b/Engine/Plugins/Experimental/PlainProps/Source/Public/PlainPropsRead.h
--- a/Engine/Plugins/Experimental/PlainProps/Source/Public/PlainPropsRead.h
+++ b/Engine/Plugins/Experimental/PlainProps/Source/Public/PlainPropsRead.h
@@ -12,8 +12,9 @@
 namespace PlainProps
 {
 
-struct FEnumSchema;
 class FMemberLoader;
+class FMemberPrinter;
+struct FEnumSchema;
 struct FSchemaBatch;
 struct FStructSchema;
 
@@ -197,11 +198,13 @@ struct FLeafView
 	FORCEINLINE char8_t		AsChar8() const		{ return As<char8_t>(); }
 	FORCEINLINE char16_t	AsChar16() const	{ return As<char16_t>(); }
 	FORCEINLINE char32_t	AsChar32() const	{ return As<char32_t>(); }
-	FORCEINLINE uint8		AsEnum8() const		{ return As<uint8,	FUnpackedLeafType{ELeafType::Enum, ELeafWidth::B8}>(); }
-	FORCEINLINE uint16		AsEnum16() const	{ return As<uint16,	FUnpackedLeafType{ELeafType::Enum, ELeafWidth::B16}>(); }
-	FORCEINLINE uint32		AsEnum32() const	{ return As<uint32,	FUnpackedLeafType{ELeafType::Enum, ELeafWidth::B32}>(); }
-	FORCEINLINE uint64		AsEnum64() const	{ return As<uint64,	FUnpackedLeafType{ELeafType::Enum, ELeafWidth::B64}>(); }
 
+	template<Arithmetic T>
+	FORCEINLINE T			AsUnderlyingValue() const
+	{
+		return As<T, FUnpackedLeafType{ELeafType::Enum, ReflectArithmetic<T>.Width}>();
+	}
+	
 	template<typename T, FUnpackedLeafType ExpectedLeaf = ReflectLeaf<T>>
 	FORCEINLINE T			As() const
 	{
@@ -283,6 +286,7 @@ public:
 
 private:
 	friend FMemberLoader;
+	friend FMemberPrinter;
 	friend FMemberReader;
 	friend FNestedRangeIterator;
 
@@ -405,20 +409,17 @@ public:
 	, Values(Data)
 	{}
 
-	ELeafType				GetType() const;
-	SIZE_T					GetSizeOf() const;
 	uint64					Num() const			{ return NumItems; }
-	uint64					NumBytes() const	{ return NumItems * GetSizeOf(); }
 
 	// These range views hide the internal representations to enable future format changes, 
 	// e.g. store zeroes or 1.0f in some compact fashion or even variable length int encodings
 	//
 	// They could also provide various conversion helpers
 	
-	template<typename T>
+	template<typename T, FUnpackedLeafType ExpectedLeaf = ReflectLeaf<T>>
 	auto					As() const
 	{
-		check(FUnpackedLeafType(Type, Width) == ReflectLeaf<T>);
+		check(FUnpackedLeafType(Type, Width) == ExpectedLeaf);
 		return TRangeView<T>(reinterpret_cast<const T*>(Values), NumItems);	
 	}
 
@@ -430,6 +431,12 @@ public:
 		return FBoolRangeView(Values, NumItems);
 	}
 
+	template<Arithmetic T>
+	auto					AsUnderlyingValues() const
+	{
+		return As<T, FUnpackedLeafType{ELeafType::Enum, ReflectArithmetic<T>.Width}>();
+	}
+
 	TRangeView< int8>		AsS8s() const		{ return As< int8>(); }
 	TRangeView<uint8>		AsU8s() const		{ return As<uint8>(); }
 	TRangeView< int16>		AsS16s() const		{ return As< int16>(); }
@@ -498,8 +505,6 @@ public:
 	TStructuralRangeView(uint64 N, FMemoryView D, SchemaType S) : NumItems(N), Data(D), Schema(S) {}
 
 	uint64			Num() const		{ return NumItems; }
-	SchemaType		GetSchema() const { return Schema; }
-
 	IteratorType	begin() const	{ return IteratorType(Schema, Data); }
 	IteratorType	end() const		{ return IteratorType(Schema, FMemoryView(Data.GetDataEnd(), 0)); }
 
@@ -539,6 +544,7 @@ class FMemberReader
 {
 public:
 	explicit FMemberReader(FStructView Struct) : FMemberReader(Struct.Schema.Resolve(), Struct.Values, Struct.Schema.Batch) {}
+	FMemberReader(const FStructSchema& Schema, FByteReader Values, FReadBatchId InBatch);
 
 	bool					HasMore() const			{ return MemberIdx < NumMembers; }
 	
@@ -580,8 +586,6 @@ protected: // for unit tests
 	const uint32			NumInnerSchemas;		// Number of static structs and enums
 #endif
 
-	FMemberReader(const FStructSchema& Schema, FByteReader Values, FReadBatchId InBatch);
-
 	const FMemberType*		GetMemberTypes() const;
 	const FMemberType*		GetRangeTypes() const;
 	const FSchemaId*		GetInnerSchemas() const;
diff --git a/Engine/Plugins/Experimental/PlainProps/Source/Public/PlainPropsTypes.h b/Engine/Plugins/Experimental/PlainProps/Source/Public/PlainPropsTypes.h
--- a/Engine/Plugins/Experimental/PlainProps/Source/Public/PlainPropsTypes.h
+++ b/Engine/Plugins/Experimental/PlainProps/Source/Public/PlainPropsTypes.h
@@ -3,6 +3,7 @@
 #pragma once
 
 #include "Containers/ArrayView.h"
+#include "Containers/StringFwd.h"
 #include "Misc/AssertionMacros.h"
 #include "Templates/TypeCompatibleBytes.h"
 #include "Templates/TypeHash.h"
@@ -419,13 +420,15 @@ public:
 	virtual FTypeId						Resolve(FEnumSchemaId Id) const = 0;
 	virtual FTypeId						Resolve(FStructSchemaId Id) const = 0;
 
-	virtual void						AppendDebugString(FString& Out, FNameId Name) const = 0;
+	virtual void						AppendString(FUtf8StringBuilderBase& Out, FNameId Name) const = 0;
 
-	PLAINPROPS_API virtual void			AppendDebugString(FString& Out, FScopeId Scope) const;
-	PLAINPROPS_API virtual void			AppendDebugString(FString& Out, FTypenameId Typename) const;
-	PLAINPROPS_API virtual void			AppendDebugString(FString& Out, FTypeId Type) const;
-	PLAINPROPS_API virtual void			AppendDebugString(FString& Out, FEnumSchemaId Name) const;
-	PLAINPROPS_API virtual void			AppendDebugString(FString& Out, FStructSchemaId Name) const;
+	PLAINPROPS_API virtual void			AppendString(FUtf8StringBuilderBase& Out, FMemberId Name) const;
+	PLAINPROPS_API virtual void			AppendString(FUtf8StringBuilderBase& Out, FOptionalMemberId Name) const;
+	PLAINPROPS_API virtual void			AppendString(FUtf8StringBuilderBase& Out, FScopeId Scope) const;
+	PLAINPROPS_API virtual void			AppendString(FUtf8StringBuilderBase& Out, FTypenameId Typename) const;
+	PLAINPROPS_API virtual void			AppendString(FUtf8StringBuilderBase& Out, FTypeId Type) const;
+	PLAINPROPS_API virtual void			AppendString(FUtf8StringBuilderBase& Out, FEnumSchemaId Name) const;
+	PLAINPROPS_API virtual void			AppendString(FUtf8StringBuilderBase& Out, FStructSchemaId Name) const;
 
 	PLAINPROPS_API FString				Print(FNameId Name) const;
 	PLAINPROPS_API FString				Print(FMemberId Name) const;
diff --git a/Engine/Plugins/Experimental/PlainProps/Source/Public/PlainPropsUeCoreBindings.h b/Engine/Plugins/Experimental/PlainProps/Source/Public/PlainPropsUeCoreBindings.h
--- a/Engine/Plugins/Experimental/PlainProps/Source/Public/PlainPropsUeCoreBindings.h
+++ b/Engine/Plugins/Experimental/PlainProps/Source/Public/PlainPropsUeCoreBindings.h
@@ -948,7 +948,7 @@ struct TTypename<TMap<K, V, SetAllocator, KeyFuncs>>
 };
 
 template<>
-PLAINPROPS_API void AppendString(FString& Out, const FName& Name);
+PLAINPROPS_API void AppendString(FUtf8StringBuilderBase& Out, const FName& Name);
 
 template<typename T, typename Allocator>
 struct TRangeBind<TArray<T, Allocator>>
@@ -1031,4 +1031,4 @@ struct TCustomBind<TPair<K, V>>
 template<typename K, typename V>
 struct TOccupancyOf<TPair<K, V>> : FRequireAll {};
 
-} // namespace PlainProps
\ No newline at end of file
+} // namespace PlainProps
diff --git a/Engine/Plugins/Experimental/PlainProps/Source/Public/PlainPropsWrite.h b/Engine/Plugins/Experimental/PlainProps/Source/Public/PlainPropsWrite.h
--- a/Engine/Plugins/Experimental/PlainProps/Source/Public/PlainPropsWrite.h
+++ b/Engine/Plugins/Experimental/PlainProps/Source/Public/PlainPropsWrite.h
@@ -11,7 +11,6 @@ namespace PlainProps
 
 struct FBuiltSchemas;
 struct FBuiltStruct;
-class FDebugIds;
 class FIdIndexerBase;
 class FNestedScopeIndexer;
 class FParametricTypeIndexer;
@@ -38,4 +37,4 @@ private:
 	TUniquePtr<FWriteIds>							NewIds;
 };
 
-} // namespace PlainProps
\ No newline at end of file
+} // namespace PlainProps
diff --git a/Engine/Plugins/Experimental/PlainProps/Source/Tests/PlainPropsTests.cpp b/Engine/Plugins/Experimental/PlainProps/Source/Tests/PlainPropsTests.cpp
--- a/Engine/Plugins/Experimental/PlainProps/Source/Tests/PlainPropsTests.cpp
+++ b/Engine/Plugins/Experimental/PlainProps/Source/Tests/PlainPropsTests.cpp
@@ -2,10 +2,12 @@
 
 #if WITH_TESTS
 #include "Algo/Compare.h"
+#include "Logging/LogMacros.h"
 #include "Tests/TestHarnessAdapter.h"
 #include "Containers/AnsiString.h"
 #include "Containers/Map.h"
 #include "Containers/Set.h"
+#include "Containers/StringConv.h"
 #include "Containers/StringView.h"
 #include "PlainPropsRead.h"
 #include "PlainPropsWrite.h"
@@ -13,8 +15,12 @@
 #include "PlainPropsIndex.h"
 #include "PlainPropsInternalBuild.h"
 #include "PlainPropsInternalFormat.h"
+#include "PlainPropsInternalPrint.h"
+#include "PlainPropsInternalRead.h"
 #include "Templates/UnrealTemplate.h"
 
+DEFINE_LOG_CATEGORY_STATIC(LogPlainPropsTests, Log, All);
+
 namespace PlainProps
 {
 	
@@ -112,6 +118,12 @@ TEST_CASE_NAMED(FPlainPropsIndexTest, "System::Core::Serialization::PlainProps::
 
 inline constexpr uint32 TestMagics[] = {0xFEEDF00D, 0xABCD1234, 0xDADADAAA, 0x99887766, 0xF0F1F2F3 };
 
+template<>
+void AppendString(FUtf8StringBuilderBase& Out, const FAnsiString& Str)
+{
+	Out.Append(Str);
+}
+
 class FTestBatchBuilder : public TIdIndexer<FAnsiString>, public IStructBindIds
 {
 public:
@@ -296,6 +308,40 @@ private:
 	TArray<const char*> Names;
 };
 
+class FTestBatchNames : public FDebugIds
+{
+public:
+	FTestBatchNames(const FTestNameReader& InNames, const FSchemaBatch& InSchemas) : Names(InNames), Schemas(InSchemas) {}
+
+	virtual FNestedScope			Resolve(FNestedScopeId Id) const override final		{ return ResolveNestedScope(Schemas, Id); }
+	virtual FParametricTypeView		Resolve(FParametricTypeId Id) const override final	{ return ResolveParametricType(Schemas, Id); }
+	virtual FTypeId					Resolve(FEnumSchemaId Id) const override final		{ return ResolveEnumSchema(Schemas, Id).Type; }
+	virtual FTypeId					Resolve(FStructSchemaId Id) const override final	{ return ResolveStructSchema(Schemas, Id).Type; }
+
+	using FDebugIds::AppendString;
+	virtual void AppendString(FUtf8StringBuilderBase& Out, FNameId Name) const override final { Out.Append(Names[Name]); }
+
+private:
+	const FTestNameReader& Names;
+	const FSchemaBatch& Schemas;
+};
+
+static void PrintBatch(
+	const FTestNameReader& Names,
+	FReadBatchId Batch,
+	const FSchemaBatch& Schemas,
+	TConstArrayView<FStructView> Objects)
+{
+	FTestBatchNames BatchNames(Names, Schemas);
+	TUtf8StringBuilder<4096> YamlString;
+	PrintBatchAsYamlDoc(YamlString, BatchNames, Batch, Objects);
+	FUtf8StringView YamlView = YamlString.ToView();
+	auto Wide = StringCast<TCHAR>(YamlView.GetData(), YamlView.Len());
+	UE_LOG(LogPlainPropsTests, Log, TEXT("\n%.*s\n"), Wide.Length(), Wide.Get());
+}
+
+////////////////////////////////////////////////////////////////////////////////////////////////
+
 class FTestBatchReader
 {
 public:
@@ -325,6 +371,9 @@ public:
 		
 		CHECK(It.Grab<uint32>() == TestMagics[4]);
 		CHECK(!Objects.IsEmpty());
+
+		// Print schemas and objects
+		PrintBatch(Names, Batch, *Schemas, Objects);
 	}
 
 	~FTestBatchReader()
@@ -332,7 +381,7 @@ public:
 		UnmountReadSchemas(Objects[0].Schema.Batch);
 	}
 
-	TConstArrayView<FStructView>		GetObjects() const { return Objects;	}
+	TConstArrayView<FStructView>		GetObjects() const { return Objects; }
 	const FTestNameReader&				GetNames() const { return Names; }
 	
 private:
@@ -643,7 +692,7 @@ TEST_CASE_NAMED(FPlainPropsReadWriteTest, "System::Core::Serialization::PlainPro
 			FStructSchemaId NestedId = Batch.DeclareStruct("Test", "Nested", NestedMembers, EMemberPresence::AllowSparse);
 
 			FMemberBuilder Members;
-			Members.Add(	Batch.NameMember("I1"), 100);
+			Members.Add(Batch.NameMember("I1"), 100);
 			FBuiltStruct* NestedInStruct = Members.BuildAndReset(Scratch, Batch.Get(NestedId), Batch);
 
 			Members.AddStruct(Batch.NameMember("Nested"), NestedId, MoveTemp(NestedInStruct));
@@ -681,6 +730,8 @@ TEST_CASE_NAMED(FPlainPropsReadWriteTest, "System::Core::Serialization::PlainPro
 
 	SECTION("Enum")
 	{
+		enum class EFlatSparse8 : int8 { A = 1, B = 2, C = 3 };
+
 		static constexpr std::initializer_list<const char*> MemberNames = 
 		{ "A2", "A0", "B0", "B4", "B5", "B7", "C3", "D34", "Max8", "Max16", "Max32", "Max64", "IF" };
 
@@ -709,7 +760,7 @@ TEST_CASE_NAMED(FPlainPropsReadWriteTest, "System::Core::Serialization::PlainPro
 			Members.AddEnum8(Batch.NameMember("B4"),		B, 4);
 			Members.AddEnum8(Batch.NameMember("B5"),		B, 5);
 			Members.AddEnum8(Batch.NameMember("B7"),		B, 7);
-			Members.AddEnum8(Batch.NameMember("C3"),		C, 3);
+			Members.AddEnum(Batch.NameMember("C3"),			C, EFlatSparse8::C);
 			Members.AddEnum8(Batch.NameMember("D34"),		D, 34);
 			Members.AddEnum8(Batch.NameMember("Max8"),		E, 0xFF);
 			Members.AddEnum16(Batch.NameMember("Max16"),	F, 0xFFFF);
@@ -731,19 +782,19 @@ TEST_CASE_NAMED(FPlainPropsReadWriteTest, "System::Core::Serialization::PlainPro
 			};
 
 			FTestMemberReader It1(Objects[0]);
-			CHECK(It1.GrabLeaf().AsEnum8() == 2);
-			CHECK(It1.GrabLeaf().AsEnum8() == 0);
-			CHECK(It1.GrabLeaf().AsEnum8() == 0);
-			CHECK(It1.GrabLeaf().AsEnum8() == 4);
-			CHECK(It1.GrabLeaf().AsEnum8() == 5);
-			CHECK(It1.GrabLeaf().AsEnum8() == 7);
-			CHECK(It1.GrabLeaf().AsEnum8() == 3);
-			CHECK(It1.GrabLeaf().AsEnum8() == 34);
-			CHECK(It1.GrabLeaf().AsEnum8() == 0xFF);
-			CHECK(It1.GrabLeaf().AsEnum16() == 0xFFFF);
-			CHECK(It1.GrabLeaf().AsEnum32() == 0xFFFFFFFF);
-			CHECK(It1.GrabLeaf().AsEnum64() == 0xFFFFFFFFFFFFFFFF);
-			CHECK(It1.GrabLeaf().AsEnum64() == 0x8000000000000001);	
+			CHECK(It1.GrabLeaf().AsUnderlyingValue<uint8>() == 2);
+			CHECK(It1.GrabLeaf().AsUnderlyingValue<uint8>() == 0);
+			CHECK(It1.GrabLeaf().AsUnderlyingValue<uint8>() == 0);
+			CHECK(It1.GrabLeaf().AsUnderlyingValue<uint8>() == 4);
+			CHECK(It1.GrabLeaf().AsUnderlyingValue<uint8>() == 5);
+			CHECK(It1.GrabLeaf().AsUnderlyingValue<uint8>() == 7);
+			CHECK(It1.GrabLeaf().As<EFlatSparse8>() == EFlatSparse8::C);
+			CHECK(It1.GrabLeaf().AsUnderlyingValue<uint8>() == 34);
+			CHECK(It1.GrabLeaf().AsUnderlyingValue<uint8>() == 0xFF);
+			CHECK(It1.GrabLeaf().AsUnderlyingValue<uint16>() == 0xFFFF);
+			CHECK(It1.GrabLeaf().AsUnderlyingValue<uint32>() == 0xFFFFFFFF);
+			CHECK(It1.GrabLeaf().AsUnderlyingValue<uint64>() == 0xFFFFFFFFFFFFFFFF);
+			CHECK(It1.GrabLeaf().AsUnderlyingValue<uint64>() == 0x8000000000000001);	
 
 			FTestMemberReader It2(Objects[0]);
 			CHECK(GetEnumName(It2.GrabLeaf()) == "FlatDense8");
@@ -798,7 +849,7 @@ TEST_CASE_NAMED(FPlainPropsReadWriteTest, "System::Core::Serialization::PlainPro
 
 			CHECK(EqualItems(GetEnumConstants<uint8>(FlatDense8), MakeConstArrayView({0, 2})));
 			CHECK(EqualItems(GetEnumConstants<uint8>(FlagDense8), TConstArrayView<uint8>()));
-			CHECK(EqualItems(GetEnumConstants<uint8>(FlatSparse8), MakeConstArrayView({3})));
+			CHECK(EqualItems(GetEnumConstants<EFlatSparse8>(FlatSparse8), MakeConstArrayView({EFlatSparse8::C})));
 			CHECK(EqualItems(GetEnumConstants<uint8>(FlagSparse8), MakeConstArrayView({2, 32})));
 			CHECK(GetEnumConstants<uint8>(FlatLimit8)[0] == 0xFF);
 			CHECK(GetEnumConstants<uint16>(FlatLimit16)[0] == 0xFFFF);
@@ -827,8 +878,8 @@ TEST_CASE_NAMED(FPlainPropsReadWriteTest, "System::Core::Serialization::PlainPro
 			FMemberBuilder Members;
 			Members.AddRange(Batch.NameMember("B0"), BuildLeafRange(Scratch, TConstArrayView<bool>()));
 			Members.AddRange(Batch.NameMember("B1"), BuildLeafRange(Scratch, MakeArrayView({true})));
-			Members.AddRange(Batch.NameMember("B8"), BuildLeafRange(Scratch, MakeArrayView({false, true, false, true, false, true, false, true})));
-			Members.AddRange(Batch.NameMember("B9"), BuildLeafRange(Scratch, MakeArrayView({true, false, true, false, true, false, true, false, true})));
+			Members.AddRange(Batch.NameMember("B8"), BuildLeafRange(Scratch, TConstArrayView<bool,int64>({false, true, false, true, false, true, false, true})));
+			Members.AddRange(Batch.NameMember("B9"), BuildLeafRange(Scratch, TConstArrayView<bool,int64>({true, false, true, false, true, false, true, false, true})));
 			Members.AddRange(Batch.NameMember("D0"), BuildLeafRange(Scratch, TConstArrayView<double>()));
 			Members.AddRange(Batch.NameMember("D3"), BuildLeafRange(Scratch, MakeArrayView({DBL_MIN, 0.0, DBL_MAX})));
 			Members.AddRange(Batch.NameMember("Hi"), BuildLeafRange(Scratch, MakeArrayView(u8"Hello!")));
@@ -863,7 +914,7 @@ TEST_CASE_NAMED(FPlainPropsReadWriteTest, "System::Core::Serialization::PlainPro
 			CHECK(EqualItems(E0.As<EUnused1>(), TConstArrayView<EUnused1>()));
 		});
 	}
-		
+
 	SECTION("StructRange")
 	{
 		TestSerialize([](FTestBatchBuilder& Batch, FScratchAllocator& Scratch)
@@ -871,9 +922,11 @@ TEST_CASE_NAMED(FPlainPropsReadWriteTest, "System::Core::Serialization::PlainPro
 			FStructSchemaId ObjectId = Batch.DeclareStruct("Test", "Object", {"Structs"}, EMemberPresence::AllowSparse);
 			FStructSchemaId StructId = Batch.DeclareStruct("Test", "Struct", {"I", "F"}, EMemberPresence::AllowSparse);
 			 
-			FStructRangeBuilder Structs(3);
-			Structs[0].Add(Batch.NameMember("I"), 1);
+			FStructRangeBuilder Structs(4);
+			Structs[0].Add(Batch.NameMember("I"), 0);
 			Structs[1].Add(Batch.NameMember("F"), 1.f);
+			Structs[2].Add(Batch.NameMember("I"), 2);
+			Structs[2].Add(Batch.NameMember("F"), 2.f);
 		
 			FMemberBuilder Members;
 			Members.AddRange(Batch.NameMember("Structs"), Structs.BuildAndReset(Scratch, Batch.Get(StructId), Batch));
@@ -886,9 +939,12 @@ TEST_CASE_NAMED(FPlainPropsReadWriteTest, "System::Core::Serialization::PlainPro
 
 			FTestMemberReader It(Objects[0]);
 			TArray<FTestMemberReader> Structs = MakeArray<FTestMemberReader>(It.GrabRange().AsStructs());
-			CHECK(Structs.Num() == 3);
-			CHECK(Structs[0].GrabLeaf().AsS32() == 1);
-			CHECK(Structs[1].GrabLeaf().AsFloat() == 1.f);	
+			CHECK(Structs.Num() == 4);
+			CHECK(Structs[0].GrabLeaf().AsS32() == 0);
+			CHECK(Structs[1].GrabLeaf().AsFloat() == 1.f);
+			CHECK(Structs[2].GrabLeaf().AsS32() == 2);
+			CHECK(Structs[2].GrabLeaf().AsFloat() == 2.f);
+			CHECK(!Structs[3].HasMore());
 		});
 	}
 	
@@ -898,7 +954,7 @@ TEST_CASE_NAMED(FPlainPropsReadWriteTest, "System::Core::Serialization::PlainPro
 
 		TestSerialize([](FTestBatchBuilder& Batch, FScratchAllocator& Scratch)
 		{
-			FStructSchemaId Object = Batch.DeclareStruct("Test", "Object", {"IntRs", "EmptyRs", "EnumRs", "StructRs", "StructRRs"}, EMemberPresence::AllowSparse);
+			FStructSchemaId Object = Batch.DeclareStruct("Test", "Object", {"IntRs", "EmptyRs", "FloatRs", "EnumRs", "StructRs", "StructRRs"}, EMemberPresence::AllowSparse);
 			FStructSchemaId XY = Batch.DeclareStruct("Test", "XY", {"X", "Y"}, EMemberPresence::RequireAll);
 			FStructSchemaId ZW = Batch.DeclareStruct("Test", "ZW", {"Z", "W"}, EMemberPresence::AllowSparse);
 			FEnumSchemaId Enum	= Batch.DeclareEnum("Test", "AB", EEnumMode::Flag, ELeafWidth::B8, {"A", "B"}, {1, 4});
@@ -908,6 +964,11 @@ TEST_CASE_NAMED(FPlainPropsReadWriteTest, "System::Core::Serialization::PlainPro
 			IntRs.Add({});
 			IntRs.Add(BuildLeafRange(Scratch, MakeArrayView({2, 3})));
 
+			FNestedRangeBuilder FloatRs(MakeLeafRangeSchema<float, int64>(), 3);
+			FloatRs.Add(BuildLeafRange(Scratch, TConstArrayView<float, int64>({1.f})));
+			FloatRs.Add({});
+			FloatRs.Add(BuildLeafRange(Scratch, TConstArrayView<float, int64>({2.f, 3.f})));
+
 			FNestedRangeBuilder EnumRs(MakeEnumRangeSchema<EAB, int32>(Enum), 2);
 			EnumRs.Add({});
 			EnumRs.Add(BuildEnumRange(Scratch, Enum, MakeArrayView({EAB::A, EAB(0), EAB::B})));
@@ -934,6 +995,7 @@ TEST_CASE_NAMED(FPlainPropsReadWriteTest, "System::Core::Serialization::PlainPro
 			FMemberBuilder Members;
 			Members.AddRange(Batch.NameMember("IntRs"), IntRs.BuildAndReset(Scratch, ERangeSizeType::S32));
 			Members.AddRange(Batch.NameMember("EmptyRs"), IntRs.BuildAndReset(Scratch, ERangeSizeType::S32));
+			Members.AddRange(Batch.NameMember("FloatRs"), FloatRs.BuildAndReset(Scratch, ERangeSizeType::S64));
 			Members.AddRange(Batch.NameMember("EnumRs"), EnumRs.BuildAndReset(Scratch, ERangeSizeType::U8));
 			Members.AddRange(Batch.NameMember("StructRs"), StructRs.BuildAndReset(Scratch, ERangeSizeType::U64));
 			Members.AddRange(Batch.NameMember("StructRRs"), StructRRs.BuildAndReset(Scratch, ERangeSizeType::U32));
@@ -947,6 +1009,7 @@ TEST_CASE_NAMED(FPlainPropsReadWriteTest, "System::Core::Serialization::PlainPro
 			FTestMemberReader It(Objects[0]);
 			TArray<FRangeView> IntRs = MakeArray<FRangeView>(It.GrabRange().AsRanges());
 			FNestedRangeView EmptyRs = It.GrabRange().AsRanges();
+			TArray<FRangeView> FloatRs = MakeArray<FRangeView>(It.GrabRange().AsRanges());
 			TArray<FRangeView> EnumRs = MakeArray<FRangeView>(It.GrabRange().AsRanges());
 			TArray<FRangeView> StructRs = MakeArray<FRangeView>(It.GrabRange().AsRanges());
 			TArray<FRangeView> StructRRs = MakeArray<FRangeView>(It.GrabRange().AsRanges());
@@ -962,6 +1025,11 @@ TEST_CASE_NAMED(FPlainPropsReadWriteTest, "System::Core::Serialization::PlainPro
 			CHECK(EnumRs[0].IsEmpty());
 			CHECK(EqualItems(EnumRs[1].AsLeaves().As<EAB>(), MakeArrayView({EAB::A, EAB(0), EAB::B})));
 
+			CHECK(FloatRs.Num() == 3);
+			CHECK(EqualItems(FloatRs[0].AsLeaves().AsFloats(), MakeArrayView({1.f})));
+			CHECK(FloatRs[1].IsEmpty());
+			CHECK(EqualItems(FloatRs[2].AsLeaves().AsFloats(), MakeArrayView({2.f, 3.f})));
+
 			CHECK(StructRs.Num() == 1);
 			TArray<FTestMemberReader> XYs = MakeArray<FTestMemberReader>(StructRs[0].AsStructs());
 			CHECK(Names[XYs[0].PeekName()] == "X");
@@ -1179,6 +1247,8 @@ TEST_CASE_NAMED(FPlainPropsReadWriteTest, "System::Core::Serialization::PlainPro
 			Members.AddStruct(Batch.NameMember("TLow"), Low, Nested.BuildAndReset(Scratch, Batch.Get(Low), Batch));
 			FStructRangeBuilder NestedRange(1);
 			NestedRange[0].Add(Batch.NameMember("MInt"), 3000);
+			Nested.Add(Batch.NameMember("LInt"), 4000);
+			NestedRange[0].AddStruct(Batch.NameMember("MLow"), Low, Nested.BuildAndReset(Scratch, Batch.Get(Low), Batch));
 			Members.AddRange(Batch.NameMember("TMids"), NestedRange.BuildAndReset(Scratch, Batch.Get(Mid), Batch));
 
 			Batch.AddObject(Top, MoveTemp(Members));
@@ -1201,7 +1271,9 @@ TEST_CASE_NAMED(FPlainPropsReadWriteTest, "System::Core::Serialization::PlainPro
 			CHECK(Names[TopIt.PeekName()] == "TLow");
 			CHECK(FTestMemberReader(TopIt.GrabStruct()).GrabLeaf().AsS32() == 2000);
 			CHECK(Names[TopIt.PeekName()] == "TMids");
-			CHECK(MakeArray<FTestMemberReader>(TopIt.GrabRange().AsStructs())[0].GrabLeaf().AsS32() == 3000);
+			TArray<FTestMemberReader> MemberRangeIt = MakeArray<FTestMemberReader>(TopIt.GrabRange().AsStructs());
+			CHECK(MemberRangeIt[0].GrabLeaf().AsS32() == 3000);
+			CHECK(FTestMemberReader(MemberRangeIt[0].GrabStruct()).GrabLeaf().AsS32() == 4000);
 
 			const FStructSchema& TopSchema = TopView.Schema.Resolve();
 			const FStructSchema& MidSchema = MidView.Schema.Resolve();
@@ -1226,7 +1298,10 @@ TEST_CASE_NAMED(FPlainPropsReadWriteTest, "System::Core::Serialization::PlainPro
 			CHECK(Names[FlatIt.PeekOwner().Name] == "Top");
 			CHECK(FTestMemberReader(FlatIt.GrabStruct()).GrabLeaf().AsS32() == 2000);
 			CHECK(Names[FlatIt.PeekOwner().Name] == "Top");
-			CHECK(MakeArray<FTestMemberReader>(FlatIt.GrabRange().AsStructs())[0].GrabLeaf().AsS32() == 3000);
+			TArray<FFlatMemberReader> FlatRangeIt = MakeArray<FFlatMemberReader>(FlatIt.GrabRange().AsStructs());
+			CHECK(FlatRangeIt[0].GrabLeaf().AsS32() == 3000);
+			CHECK(FFlatMemberReader(FlatRangeIt[0].GrabStruct()).GrabLeaf().AsS32() == 4000);
+			CHECK(!FlatRangeIt[0].HasMore());
 			CHECK(!FlatIt.HasMore());
 		});
 	}
diff --git a/Engine/Plugins/Experimental/PlainProps/Source/Tests/PlainPropsUeCoreTests.cpp b/Engine/Plugins/Experimental/PlainProps/Source/Tests/PlainPropsUeCoreTests.cpp
--- a/Engine/Plugins/Experimental/PlainProps/Source/Tests/PlainPropsUeCoreTests.cpp
+++ b/Engine/Plugins/Experimental/PlainProps/Source/Tests/PlainPropsUeCoreTests.cpp
@@ -7,18 +7,24 @@
 #include "PlainPropsIndex.h"
 #include "PlainPropsInternalBuild.h"
 #include "PlainPropsInternalFormat.h"
+#include "PlainPropsInternalPrint.h"
+#include "PlainPropsInternalRead.h"
 #include "PlainPropsLoad.h"
 #include "PlainPropsRead.h"
 #include "PlainPropsSave.h"
 #include "PlainPropsWrite.h"
 #include "PlainPropsUeCoreBindings.h"
 #include "Algo/Compare.h"
+#include "Containers/StringConv.h"
 #include "Containers/StringView.h"
 #include "Containers/Map.h"
+#include "Logging/LogMacros.h"
 #include "Math/Transform.h"
 #include "Templates/UnrealTemplate.h"
 #include "Tests/TestHarnessAdapter.h"
 
+DEFINE_LOG_CATEGORY_STATIC(LogPlainPropsUeCoreTests, Log, All);
+
 namespace PlainProps::UE::Test
 {
 
@@ -252,6 +258,77 @@ TArray64<uint8> FBatchSaver::Write() const
 
 ////////////////////////////////////////////////////////////////////////////////////////////////
 
+class FSavedBatchNames : public FDebugIds
+{
+public:
+	FSavedBatchNames(TConstArrayView<FName> InNames, const FSchemaBatch& InSchemas) : Names(InNames), Schemas(InSchemas) {}
+
+	virtual FNestedScope				Resolve(FNestedScopeId Id) const override final		{ return ResolveNestedScope(Schemas, Id); }
+	virtual FParametricTypeView			Resolve(FParametricTypeId Id) const override final	{ return ResolveParametricType(Schemas, Id); }
+	virtual FTypeId						Resolve(FEnumSchemaId Id) const override final		{ return ResolveEnumSchema(Schemas, Id).Type; }
+	virtual FTypeId						Resolve(FStructSchemaId Id) const override final	{ return ResolveStructSchema(Schemas, Id).Type; }
+
+	using FDebugIds::AppendString;
+	virtual void AppendString(FUtf8StringBuilderBase& Out, FNameId Name) const override final { Names[Name.Idx].AppendString(Out); }
+
+private:
+	TConstArrayView<FName> Names;
+	const FSchemaBatch& Schemas;
+};
+
+class FLoadBatchNames : public FDebugIds
+{
+public:
+	FLoadBatchNames(const FIdBinding& InBinding, const FSchemaBatch& InSchemas) : DebugIds(GNames), Binding(InBinding), Schemas(InSchemas) {}
+
+	virtual FNestedScope				Resolve(FNestedScopeId Id) const override final		{ return DebugIds.Resolve(Binding.Remap(Id)); }
+	virtual FParametricTypeView			Resolve(FParametricTypeId Id) const override final	{ return DebugIds.Resolve(Binding.Remap(Id)); }
+	virtual FTypeId						Resolve(FEnumSchemaId Id) const override final		{ return ResolveEnumSchema(Schemas, Id).Type; }
+	virtual FTypeId						Resolve(FStructSchemaId Id) const override final	{ return ResolveStructSchema(Schemas, Id).Type; }
+
+	using FDebugIds::AppendString;
+	virtual void AppendString(FUtf8StringBuilderBase& Out, FNameId Name) const override final			{ DebugIds.AppendString(Out, Name); }
+	virtual void AppendString(FUtf8StringBuilderBase& Out, FTypenameId Typename) const override final	{ DebugIds.AppendString(Out, Typename); }
+	virtual void AppendString(FUtf8StringBuilderBase& Out, FScopeId Scope) const override final			{ DebugIds.AppendString(Out, Scope); }
+
+private:
+	const FDebugIds& DebugIds;
+	const FIdBinding& Binding;
+	const FSchemaBatch& Schemas;
+};
+
+void PrintSavedSchemas(TConstArrayView<FName> Names, const FSchemaBatch& SavedSchemas)
+{
+	FReadBatchId Batch = MountReadSchemas(&SavedSchemas);
+	TUtf8StringBuilder<4096> YamlString;
+	{
+		FYamlBuilderPtr YamlBuilder = MakeYamlBuilder(YamlString);
+		FSavedBatchNames BatchNames(Names, SavedSchemas);
+		FBatchPrinter Printer(*YamlBuilder, BatchNames);
+		Printer.PrintSchemas(Batch);
+	}
+	FUtf8StringView YamlView = YamlString.ToView();
+	auto Wide = StringCast<TCHAR>(YamlView.GetData(), YamlView.Len());
+	UE_LOG(LogPlainPropsUeCoreTests, Log, TEXT("Schemas:\n%.*s"), Wide.Length(), Wide.Get());
+	UnmountReadSchemas(Batch);
+}
+
+void PrintLoadStruct(const FIdBinding& Binding, FStructView StructView, int32 LoadIdx)
+{
+	TUtf8StringBuilder<4096> YamlString;
+	{
+		FYamlBuilderPtr YamlBuilder = MakeYamlBuilder(YamlString);
+		FLoadBatchNames BatchNames(Binding, GetReadSchemas(StructView.Schema.Batch));
+		FBatchPrinter Printer(*YamlBuilder, BatchNames);
+		Printer.PrintObjects({StructView});
+	}
+	FUtf8StringView YamlView = YamlString.ToView();
+	auto Wide = StringCast<TCHAR>(YamlView.GetData(), YamlView.Len());
+	UE_LOG(LogPlainPropsUeCoreTests, Log, TEXT("LoadStruct %d:\n%.*s"), LoadIdx, Wide.Length(), Wide.Get());
+}
+
+////////////////////////////////////////////////////////////////////////////////////////////////
+
 class FBatchLoader
 {
 public:
@@ -270,12 +347,18 @@ public:
 		uint32 SchemasSize = It.Grab<uint32>();
 		const FSchemaBatch* SavedSchemas = ValidateSchemas(It.GrabSlice(SchemasSize));
 		CHECK(It.Grab<uint32>() == Magics[2]);
+
+		// Mount, print and unmount saved schemas
+		PrintSavedSchemas(Ids, *SavedSchemas);
 		
 		// Bind saved ids to runtime ids, make new schemas with new ids and mount them
 		FIdTranslator RuntimeIds(GNames, Ids, *SavedSchemas);
 		FSchemaBatch* LoadSchemas = CreateTranslatedSchemas(*SavedSchemas, RuntimeIds.Translation);
 		FReadBatchId Batch = MountReadSchemas(LoadSchemas);
 
+		// Store mappings required for structs printing of untranslated scopes and parametric types
+		IdBinding = RuntimeIds.Translation;
+
 		// Read objects
 		while (uint64 NumBytes = It.GrabVarIntU())
 		{	
@@ -325,13 +408,16 @@ public:
 	template<class T>
 	void LoadInto(T& Out)
 	{
-		FStructView In = Objects[LoadIdx++];
+		FStructView In = Objects[LoadIdx];
+		PrintLoadStruct(IdBinding, In, LoadIdx);
 		LoadStruct(&Out, In.Values, In.Schema.Id, *Plans);
+		++LoadIdx;
 	}
 
 private:
 	TConstArrayView<FName>		Ids;
 	FNameBinding				Names;
+	FIdBinding					IdBinding;
 	FCustomBindings				Customs;
 	FLoadBatchPtr				Plans;
 	TArray<FStructView>			Objects;
@@ -1120,4 +1206,4 @@ TEST_CASE_NAMED(FPlainPropsUeCoreTest, "System::Core::Serialization::PlainProps:
 
 } // namespace PlainProps::UE::Test
 
-#endif // WITH_TESTS
\ No newline at end of file
+#endif // WITH_TESTS
diff --git a/Engine/Plugins/Experimental/PlainPropsUObject/Source/Public/PlainPropsUObjectRuntime.h b/Engine/Plugins/Experimental/PlainPropsUObject/Source/Public/PlainPropsUObjectRuntime.h
--- a/Engine/Plugins/Experimental/PlainPropsUObject/Source/Public/PlainPropsUObjectRuntime.h
+++ b/Engine/Plugins/Experimental/PlainPropsUObject/Source/Public/PlainPropsUObjectRuntime.h
@@ -8,12 +8,6 @@
 #include "PlainPropsTypes.h"
 #include "UObject/NameTypes.h"
 
-namespace PlainProps
-{
-template<>
-PLAINPROPS_API void AppendString(FString& Out, const FName& Name);
-}
-
 namespace PlainProps::UE
 {
 
@@ -82,4 +76,4 @@ struct FDeltaRuntime : FDefaultRuntime
 	static FCustomBindings&		GetCustoms()		{ return GUE.DeltaCustoms; }
 };
 
-} // namespace PlainProps::UE
\ No newline at end of file
+} // namespace PlainProps::UE
