PlainProps: Add support for noting empty enum ranges.
Fix bug in LeafRange test that used to pass silently even though the same declared enum was built with different widths.
Add NestedRange test for empty inner range enums.

#rnx

diff --git a/Engine/Plugins/Experimental/PlainProps/Source/Private/PlainPropsBuildSchema.cpp b/Engine/Plugins/Experimental/PlainProps/Source/Private/PlainPropsBuildSchema.cpp
--- a/Engine/Plugins/Experimental/PlainProps/Source/Private/PlainPropsBuildSchema.cpp
+++ b/Engine/Plugins/Experimental/PlainProps/Source/Private/PlainPropsBuildSchema.cpp
@@ -59,7 +59,7 @@ struct FStructSchemaBuilder
 	TBitArray<>									NotedMembers;
 
 	void										NoteMembersRecursively(const FBuiltStruct& Struct);
-	void										NoteRangeRecursively(ERangeSizeType NumType, TConstArrayView<FMemberType> Types, void* InnermostSchemaBuilder, const FBuiltRange& Range);
+	void										NoteRangeRecursively(ERangeSizeType NumType, TConstArrayView<FMemberType> Types, void* InnermostSchemaBuilder, const FBuiltRange* Range);
 	FBuiltStructSchema							Build() const;
 };
 
@@ -70,7 +70,8 @@ struct FEnumSchemaBuilder
 	TOptional<ELeafWidth>						NotedWidth;
 	TSet<uint64>								NotedConstants;
 
-	void										NoteValue(uint64 Value, ELeafWidth Width);
+	void										NoteValue(ELeafWidth Width, uint64 Value);
+	void										NoteEmpty(ELeafWidth Width);
 	FBuiltEnumSchema							Build() const;
 };
 
@@ -270,7 +271,7 @@ void FStructSchemaBuilder::NoteMembersRecursively(const FBuiltStruct& Struct)
 			switch (Schema.Type.GetKind())
 			{
 			case EMemberKind::Leaf:
-				AllSchemas.NoteEnum(static_cast<FEnumSchemaId>(InnerSchema)).NoteValue(Member.Value.Leaf, Schema.Type.AsLeaf().Width);
+				AllSchemas.NoteEnum(static_cast<FEnumSchemaId>(InnerSchema)).NoteValue(Schema.Type.AsLeaf().Width, Member.Value.Leaf);
 				break;
 
 			case EMemberKind::Struct:
@@ -279,10 +280,7 @@ void FStructSchemaBuilder::NoteMembersRecursively(const FBuiltStruct& Struct)
 		
 			case EMemberKind::Range:
 				void* InnerSchemaBuilder = NoteStructOrEnum(AllSchemas, Schema.GetInnermostType().IsStruct(), InnerSchema);
-				if (Member.Value.Range)
-				{
-					NoteRangeRecursively(Schema.Type.AsRange().MaxSize, Schema.GetInnerRangeTypes(), InnerSchemaBuilder, *Member.Value.Range);			
-				}
+				NoteRangeRecursively(Schema.Type.AsRange().MaxSize, Schema.GetInnerRangeTypes(), InnerSchemaBuilder, Member.Value.Range);			
 				break;
 			}
 		}
@@ -294,7 +292,7 @@ void NoteEnumValues(FEnumSchemaBuilder& Schema, const IntType* Values, uint64 Nu
 {
 	for (IntType Value : TConstArrayView64<IntType>(Values, Num))
 	{
-		Schema.NoteValue(Value, LeafWidth<sizeof(IntType)>);
+		Schema.NoteValue(LeafWidth<sizeof(IntType)>, Value);
 	}
 }
 
@@ -310,28 +308,40 @@ static void NoteEnumRange(FEnumSchemaBuilder& Out, FLeafType Leaf, const FBuiltR
 	}
 }
 
-void FStructSchemaBuilder::NoteRangeRecursively(ERangeSizeType NumType, TConstArrayView<FMemberType> Types, void* InnermostSchema, const FBuiltRange& Range)
+static void NoteEmptyRange(ERangeSizeType NumType, TConstArrayView<FMemberType> Types, void* InnermostSchema)
 {
+	FMemberType InnermostType = Types.Last();
+	if (IsEnum(InnermostType))
+	{
+		static_cast<FEnumSchemaBuilder*>(InnermostSchema)->NoteEmpty(InnermostType.AsLeaf().Width);
+	}
+}
+
+void FStructSchemaBuilder::NoteRangeRecursively(ERangeSizeType NumType, TConstArrayView<FMemberType> Types, void* InnermostSchema, const FBuiltRange* Range)
+{
+	if (Range == nullptr)
+	{
+		NoteEmptyRange(NumType, Types, InnermostSchema);
+		return;
+	}
+	checkf(Range->Num > 0, TEXT("Range was built but without values"));
 	FMemberType Type = Types[0];
 	switch (Type.GetKind())
 	{
 	case EMemberKind::Struct:
-		for (const FBuiltStruct* Struct : Range.AsStructs())
+		for (const FBuiltStruct* Struct : Range->AsStructs())
 		{
 			static_cast<FStructSchemaBuilder*>(InnermostSchema)->NoteMembersRecursively(*Struct);
 		}
 		break;
 	case EMemberKind::Range:
-		for (const FBuiltRange* InnerRange : Range.AsRanges())
+		for (const FBuiltRange* InnerRange : Range->AsRanges())
 		{
-			if (InnerRange)
-			{
-				NoteRangeRecursively(Type.AsRange().MaxSize, Types.RightChop(1), InnermostSchema, *InnerRange);
-			}
+			NoteRangeRecursively(Type.AsRange().MaxSize, Types.RightChop(1), InnermostSchema, InnerRange);
 		}
 		break;
 	case EMemberKind::Leaf:
-		NoteEnumRange(/* out */ *static_cast<FEnumSchemaBuilder*>(InnermostSchema), Type.AsLeaf(), Range);
+		NoteEnumRange(/* out */ *static_cast<FEnumSchemaBuilder*>(InnermostSchema), Type.AsLeaf(), *Range);
 		break;
 	}
 }
@@ -390,7 +400,7 @@ FBuiltEnumSchema FEnumSchemaBuilder::Build() const
 	return Out;
 }
 
-void FEnumSchemaBuilder::NoteValue(uint64 Value, ELeafWidth Width)
+void FEnumSchemaBuilder::NoteValue(ELeafWidth Width, uint64 Value)
 {
 	check(NotedWidth == NullOpt || NotedWidth == Width);
 	NotedWidth = Width;
@@ -430,4 +440,10 @@ void FEnumSchemaBuilder::NoteValue(uint64 Value, ELeafWidth Width)
 	}
 }
 
+void FEnumSchemaBuilder::NoteEmpty(ELeafWidth Width)
+{
+	check(NotedWidth == NullOpt || NotedWidth == Width);
+	NotedWidth = Width;
+}
+
 } // namespace PlainProps
diff --git a/Engine/Plugins/Experimental/PlainProps/Source/Tests/PlainPropsTests.cpp b/Engine/Plugins/Experimental/PlainProps/Source/Tests/PlainPropsTests.cpp
--- a/Engine/Plugins/Experimental/PlainProps/Source/Tests/PlainPropsTests.cpp
+++ b/Engine/Plugins/Experimental/PlainProps/Source/Tests/PlainPropsTests.cpp
@@ -884,7 +884,7 @@ TEST_CASE_NAMED(FPlainPropsReadWriteTest, "System::Core::Serialization::PlainPro
 			Members.AddRange(Batch.NameMember("D3"), BuildLeafRange(Scratch, MakeArrayView({DBL_MIN, 0.0, DBL_MAX})));
 			Members.AddRange(Batch.NameMember("Hi"), BuildLeafRange(Scratch, MakeArrayView(u8"Hello!")));
 			Members.AddRange(Batch.NameMember("E3"), BuildEnumRange(Scratch, Enum, MakeArrayView({EABCD::B, EABCD::A, EABCD::D})));
-			Members.AddRange(Batch.NameMember("E0"), BuildEnumRange(Scratch, Enum, TConstArrayView<EUnused1>()));
+			Members.AddRange(Batch.NameMember("E0"), BuildEnumRange(Scratch, UnusedEnum1, TConstArrayView<EUnused1>()));
 
 			Batch.AddObject(ObjectId, MoveTemp(Members));
 		}, 
@@ -951,13 +951,15 @@ TEST_CASE_NAMED(FPlainPropsReadWriteTest, "System::Core::Serialization::PlainPro
 	SECTION("NestedRange")
 	{
 		enum class EAB : uint8 { A = 1, B = 4 };
+		enum class EUnused : uint8 { X };
 
 		TestSerialize([](FTestBatchBuilder& Batch, FScratchAllocator& Scratch)
 		{
-			FStructSchemaId Object = Batch.DeclareStruct("Test", "Object", {"IntRs", "EmptyRs", "FloatRs", "EnumRs", "StructRs", "StructRRs"}, EMemberPresence::AllowSparse);
+			FStructSchemaId Object = Batch.DeclareStruct("Test", "Object", {"IntRs", "EmptyRs", "FloatRs", "EnumRs", "UnusedEnumRs", "StructRs", "StructRRs"}, EMemberPresence::AllowSparse);
 			FStructSchemaId XY = Batch.DeclareStruct("Test", "XY", {"X", "Y"}, EMemberPresence::RequireAll);
 			FStructSchemaId ZW = Batch.DeclareStruct("Test", "ZW", {"Z", "W"}, EMemberPresence::AllowSparse);
-			FEnumSchemaId Enum	= Batch.DeclareEnum("Test", "AB", EEnumMode::Flag, ELeafWidth::B8, {"A", "B"}, {1, 4});
+			FEnumSchemaId Enum = Batch.DeclareEnum("Test", "AB", EEnumMode::Flag, ELeafWidth::B8, {"A", "B"}, {1, 4});
+			FEnumSchemaId UnusedEnum = Batch.DeclareEnum("Test", "Unused", EEnumMode::Flat, ELeafWidth::B8, {"X"}, {0});
 	
 			FNestedRangeBuilder IntRs(MakeLeafRangeSchema<int32, int32>(), 3);
 			IntRs.Add(BuildLeafRange(Scratch, MakeArrayView({1})));
@@ -973,6 +975,10 @@ TEST_CASE_NAMED(FPlainPropsReadWriteTest, "System::Core::Serialization::PlainPro
 			EnumRs.Add({});
 			EnumRs.Add(BuildEnumRange(Scratch, Enum, MakeArrayView({EAB::A, EAB(0), EAB::B})));
 
+			FNestedRangeBuilder UnusedEnumRs(MakeEnumRangeSchema<EUnused, int32>(UnusedEnum), 2);
+			UnusedEnumRs.Add({});
+			UnusedEnumRs.Add(BuildEnumRange(Scratch, UnusedEnum, TConstArrayView<EUnused>()));
+
 			FStructRangeBuilder XYs(uint64(2));
 			XYs[0].Add(Batch.NameMember("X"), 1.f);
 			XYs[0].Add(Batch.NameMember("Y"), 2.f);
@@ -997,6 +1003,7 @@ TEST_CASE_NAMED(FPlainPropsReadWriteTest, "System::Core::Serialization::PlainPro
 			Members.AddRange(Batch.NameMember("EmptyRs"), IntRs.BuildAndReset(Scratch, ERangeSizeType::S32));
 			Members.AddRange(Batch.NameMember("FloatRs"), FloatRs.BuildAndReset(Scratch, ERangeSizeType::S64));
 			Members.AddRange(Batch.NameMember("EnumRs"), EnumRs.BuildAndReset(Scratch, ERangeSizeType::U8));
+			Members.AddRange(Batch.NameMember("UnusedEnumRs"), UnusedEnumRs.BuildAndReset(Scratch, ERangeSizeType::U8));
 			Members.AddRange(Batch.NameMember("StructRs"), StructRs.BuildAndReset(Scratch, ERangeSizeType::U64));
 			Members.AddRange(Batch.NameMember("StructRRs"), StructRRs.BuildAndReset(Scratch, ERangeSizeType::U32));
 
@@ -1011,6 +1018,7 @@ TEST_CASE_NAMED(FPlainPropsReadWriteTest, "System::Core::Serialization::PlainPro
 			FNestedRangeView EmptyRs = It.GrabRange().AsRanges();
 			TArray<FRangeView> FloatRs = MakeArray<FRangeView>(It.GrabRange().AsRanges());
 			TArray<FRangeView> EnumRs = MakeArray<FRangeView>(It.GrabRange().AsRanges());
+			TArray<FRangeView> UnusedEnumRs = MakeArray<FRangeView>(It.GrabRange().AsRanges());
 			TArray<FRangeView> StructRs = MakeArray<FRangeView>(It.GrabRange().AsRanges());
 			TArray<FRangeView> StructRRs = MakeArray<FRangeView>(It.GrabRange().AsRanges());
 
@@ -1025,6 +1033,10 @@ TEST_CASE_NAMED(FPlainPropsReadWriteTest, "System::Core::Serialization::PlainPro
 			CHECK(EnumRs[0].IsEmpty());
 			CHECK(EqualItems(EnumRs[1].AsLeaves().As<EAB>(), MakeArrayView({EAB::A, EAB(0), EAB::B})));
 
+			CHECK(UnusedEnumRs.Num() == 2);
+			CHECK(UnusedEnumRs[0].IsEmpty());
+			CHECK(EqualItems(UnusedEnumRs[1].AsLeaves().As<EUnused>(), TConstArrayView<EUnused>()));
+
 			CHECK(FloatRs.Num() == 3);
 			CHECK(EqualItems(FloatRs[0].AsLeaves().AsFloats(), MakeArrayView({1.f})));
 			CHECK(FloatRs[1].IsEmpty());
