PlainProps: Generate serialized name indices in the deterministic order of traversing the built schemas

#rnx

diff --git a/Engine/Plugins/Experimental/PlainProps/Source/Private/PlainPropsWrite.cpp b/Engine/Plugins/Experimental/PlainProps/Source/Private/PlainPropsWrite.cpp
--- a/Engine/Plugins/Experimental/PlainProps/Source/Private/PlainPropsWrite.cpp
+++ b/Engine/Plugins/Experimental/PlainProps/Source/Private/PlainPropsWrite.cpp
@@ -35,6 +35,7 @@ struct FWriteIds
 	TArray<FNestedScope>				KeptScopes;
 	TArray<FParametricType>				KeptParametrics;
 	TArray<FTypeId>						KeptParameters;
+	TArray<FNameId>						KeptNames;
 
 	FNameId								Remap(FNameId Old) const				{ return Names[Old.Idx].Get(); }
 	FMemberId							Remap(FMemberId Old) const				{ return { Remap(Old.Id) }; }
@@ -69,6 +70,7 @@ struct FUsedIds
 	TBitArray<> Names;
 	TBitArray<> NestedScopes;
 	TBitArray<> ParametricTypes;
+	TArray<FNameId> KeptNames;
 
 
 	explicit FUsedIds(const FIdIndexerBase& InIds)
@@ -76,7 +78,9 @@ struct FUsedIds
 	, Names(false, Ids.NumNames())
 	, NestedScopes(false, Ids.GetNestedScopes().Num())
 	, ParametricTypes(false, Ids.GetParametricTypes().Num())
-	{}
+	{
+		KeptNames.Reserve(16 * (Ids.NumEnums() + Ids.NumStructs()));
+	}
 
 	template<class PartialSchemaType>
 	void DetectUsage(const TArray<PartialSchemaType>& Schemas)
@@ -94,7 +98,11 @@ struct FUsedIds
 
 	void MarkUsed(FNameId Name)
 	{
-		Names[Name.Idx] = true;
+		if (!Names[Name.Idx])
+		{
+			Names[Name.Idx] = true;
+			KeptNames.Add(Name);
+		}
 	}
 	
 	void MarkUsed(FOptionalConcreteTypenameId Name)
@@ -212,11 +220,17 @@ FWriteIds::FWriteIds(const FIdIndexerBase& Ids, const IStructBindIds& InBindIds,
 
 	if (Format == ESchemaFormat::StableNames)
 	{
-		// Generate new name indices
+		// Generate new name indices in the deterministic order of traversing the built schemas
 		FUsedIds Used(Ids);
 		Used.DetectUsage(Schemas.Structs);
 		Used.DetectUsage(Schemas.Enums);
-		MakeRemapping(Names, Used.Names);
+
+		KeptNames = MoveTemp(Used.KeptNames);
+		Names.SetNum(Ids.NumNames());
+		for (uint32 Idx = 0, Num = IntCastChecked<uint32>(KeptNames.Num()); Idx < Num; ++Idx)
+		{
+			Names[KeptNames[Idx].Idx] = FNameId{ Idx };
+		}
 
 		// Generate new nested scope and parametric type mappings
 		KeptScopes.Reserve(MakeRemapping(/* out */ NestedScopes, Used.NestedScopes));
@@ -781,9 +795,9 @@ FWriter::FWriter(const FIdIndexerBase& AllIds, const IStructBindIds& BindIds, co
 FWriter::~FWriter()
 {}
 
-bool FWriter::Uses(FNameId BuiltId) const
+TConstArrayView<FNameId> FWriter::GetUsedNames() const
 {
-	return !!NewIds->Names[BuiltId.Idx];
+	return NewIds->KeptNames;
 }
 
 FOptionalStructSchemaId FWriter::GetWriteId(FStructSchemaId BuiltId) const
diff --git a/Engine/Plugins/Experimental/PlainProps/Source/Public/PlainPropsWrite.h b/Engine/Plugins/Experimental/PlainProps/Source/Public/PlainPropsWrite.h
--- a/Engine/Plugins/Experimental/PlainProps/Source/Public/PlainPropsWrite.h
+++ b/Engine/Plugins/Experimental/PlainProps/Source/Public/PlainPropsWrite.h
@@ -25,7 +25,7 @@ public:
 	PLAINPROPS_API FWriter(const FIdIndexerBase& DeclaredIds, const IStructBindIds& BindIds, const FBuiltSchemas& InSchemas, ESchemaFormat Format);
 	PLAINPROPS_API ~FWriter();
 	
-	PLAINPROPS_API bool								Uses(FNameId BuiltId) const;
+	PLAINPROPS_API TConstArrayView<FNameId>			GetUsedNames() const;
 	PLAINPROPS_API FOptionalStructSchemaId			GetWriteId(FStructSchemaId BuiltId) const;
 
 	PLAINPROPS_API void								WriteSchemas(TArray64<uint8>& Out) const;
diff --git a/Engine/Plugins/Experimental/PlainProps/Source/Tests/PlainPropsTests.cpp b/Engine/Plugins/Experimental/PlainProps/Source/Tests/PlainPropsTests.cpp
--- a/Engine/Plugins/Experimental/PlainProps/Source/Tests/PlainPropsTests.cpp
+++ b/Engine/Plugins/Experimental/PlainProps/Source/Tests/PlainPropsTests.cpp
@@ -224,13 +224,10 @@ TArray64<uint8> FTestBatchBuilder::Write()
 	// Write names
 	TArray64<uint8> Out;
 	TArray64<uint8> Tmp;
-	uint32 OldNameIdx = 0;
-	for (const FAnsiString& Name : Names)
+	for (FNameId Name : Writer.GetUsedNames())
 	{
-		if (Writer.Uses({OldNameIdx++}))
-		{
-			WriteData(Tmp, *Name, Name.Len() + 1);
-		}
+		const FAnsiString& Str = ResolveName(Name);
+		WriteData(Tmp, *Str, Str.Len() + 1);
 	}
 	WriteU32(Out, TestMagics[0]);
 	WriteSkippableSlice(Out, Tmp);
diff --git a/Engine/Plugins/Experimental/PlainProps/Source/Tests/PlainPropsUeCoreTests.cpp b/Engine/Plugins/Experimental/PlainProps/Source/Tests/PlainPropsUeCoreTests.cpp
--- a/Engine/Plugins/Experimental/PlainProps/Source/Tests/PlainPropsUeCoreTests.cpp
+++ b/Engine/Plugins/Experimental/PlainProps/Source/Tests/PlainPropsUeCoreTests.cpp
@@ -213,12 +213,10 @@ TArray64<uint8> FBatchSaver::Write() const
 	// Filter out declared but unused names and ids
 	FWriter Writer(GNames, BindIds, Schemas, ESchemaFormat::StableNames);
 	TArray<FName> UsedNames;
-	for (uint32 Idx = 0, Num = GNames.NumNames(); Idx < Num; ++Idx)
+	UsedNames.Reserve(Writer.GetUsedNames().Num());
+	for (FNameId Name : Writer.GetUsedNames())
 	{
-		if (Writer.Uses({Idx}))
-		{
-			UsedNames.Add(GNames.ResolveName({Idx}));
-		}
+		UsedNames.Add(GNames.ResolveName(Name));
 	}
 	
 	// Write ids. Just copying in-memory FNames, a stable format might use SaveNameBatch().
