PlainProps: Fix incorrect size calcuation for FStructSchema.
Add size calculation checks when writing FStructSchema and FEnumSchema.

#rnx

diff --git a/Engine/Plugins/Experimental/PlainProps/Source/Private/PlainPropsInternalFormat.h b/Engine/Plugins/Experimental/PlainProps/Source/Private/PlainPropsInternalFormat.h
--- a/Engine/Plugins/Experimental/PlainProps/Source/Private/PlainPropsInternalFormat.h
+++ b/Engine/Plugins/Experimental/PlainProps/Source/Private/PlainPropsInternalFormat.h
@@ -106,8 +106,8 @@ static_assert(offsetof(FStructSchema, Footer) == 15,							"Add binary format ve
 inline uint32 CalculateSize(const FStructSchema& In)
 {
 	static_assert(alignof(FMemberType) == 1);
-	uint32 Out = sizeof(In) + sizeof(FMemberType) * (uint32(In.NumMembers) + In.NumRangeTypes);
-	Out = Align(Out + In.NumMembers * sizeof(FMemberId), sizeof(FMemberId));
+	uint32 Out = offsetof(In, Footer) + sizeof(FMemberType) * (uint32(In.NumMembers) + In.NumRangeTypes);
+	Out = Align(Out + In.NumNames() * sizeof(FMemberId), sizeof(FMemberId));
 	Out = Align(Out + (In.NumInnerSchemas) * sizeof(FSchemaId), sizeof(FSchemaId));
 	return Out;
 }
@@ -132,7 +132,7 @@ static_assert(offsetof(FEnumSchema, Footer) == 12,						"Add binary format versi
 
 inline uint32 CalculateSize(const FEnumSchema& In)
 {
-	return Align(sizeof(In) + In.Num * sizeof(In.Footer[0]) + In.ExplicitConstants * In.Num * SizeOf(In.Width), alignof(FEnumSchema));
+	return Align(offsetof(In, Footer) + In.Num * sizeof(In.Footer[0]) + In.ExplicitConstants * In.Num * SizeOf(In.Width), alignof(FEnumSchema));
 }
 
 //////////////////////////////////////////////////////////////////////////
diff --git a/Engine/Plugins/Experimental/PlainProps/Source/Private/PlainPropsWrite.cpp b/Engine/Plugins/Experimental/PlainProps/Source/Private/PlainPropsWrite.cpp
--- a/Engine/Plugins/Experimental/PlainProps/Source/Private/PlainPropsWrite.cpp
+++ b/Engine/Plugins/Experimental/PlainProps/Source/Private/PlainPropsWrite.cpp
@@ -389,12 +389,15 @@ static void WriteSchema(TArray64<uint8>& Out, const FBuiltStructSchema& Struct,
 	BinaryHeader.NumRangeTypes = IntCastChecked<uint16>(InnerRangeTypes.Num());
 	BinaryHeader.NumInnerSchemas = IntCastChecked<uint16>(InnerSchemas.Num());
 
-	check(IsAligned(Out.Num(), alignof(FStructSchema)));
+	const int64 Num = Out.Num(); (void)Num;
+	check(IsAligned(Num, alignof(FStructSchema)));
 	WriteData(Out, HeaderData, sizeof(HeaderData));
 	WriteArray(Out, MemberTypes);
 	WriteArray(Out, InnerRangeTypes);
 	WriteAlignedArray(Out, MakeArrayView(MemberNames));
 	WriteAlignedArray(Out, MakeArrayView(InnerSchemas));
+	check(IsAligned(Out.Num(), alignof(FStructSchema)));
+	check(Out.Num() - Num == CalculateSize(*reinterpret_cast<FStructSchema*>(Out.GetData() + Num)));
 }
 
 static bool IsFlatSequence(TConstArrayView<uint64> Constants)
@@ -458,7 +461,8 @@ static void WriteSchema(TArray64<uint8>& Out, const FBuiltEnumSchema& Enum, cons
 	BinaryHeader.Width = Enum.Width;
 	BinaryHeader.Num = IntCastChecked<uint16>(Names.Num());
 
-	check(IsAligned(Out.Num(), alignof(FEnumSchema)));
+	const int64 Num = Out.Num(); (void)Num;
+	check(IsAligned(Num, alignof(FEnumSchema)));
 	WriteData(Out, &BinaryHeader, sizeof(BinaryHeader));
 	WriteArray(Out, Names);
 	if (BinaryHeader.ExplicitConstants)
@@ -466,6 +470,7 @@ static void WriteSchema(TArray64<uint8>& Out, const FBuiltEnumSchema& Enum, cons
 		WriteEnumConstants(Out, Enum.Width, Enum.Constants);
 	}
 	WriteAlignmentPadding<FEnumSchema>(Out);
+	check(Out.Num() - Num == CalculateSize(*reinterpret_cast<FEnumSchema*>(Out.GetData() + Num)));
 }
 
 template<typename T>
