PlainProps: Fix review nitpicks for fixing incorrect size calcuation (see cl 38328349)

#rnx

diff --git a/Engine/Plugins/Experimental/PlainProps/Source/Private/PlainPropsWrite.cpp b/Engine/Plugins/Experimental/PlainProps/Source/Private/PlainPropsWrite.cpp
--- a/Engine/Plugins/Experimental/PlainProps/Source/Private/PlainPropsWrite.cpp
+++ b/Engine/Plugins/Experimental/PlainProps/Source/Private/PlainPropsWrite.cpp
@@ -389,15 +389,15 @@ static void WriteSchema(TArray64<uint8>& Out, const FBuiltStructSchema& Struct,
 	BinaryHeader.NumRangeTypes = IntCastChecked<uint16>(InnerRangeTypes.Num());
 	BinaryHeader.NumInnerSchemas = IntCastChecked<uint16>(InnerSchemas.Num());
 
-	const int64 Num = Out.Num(); (void)Num;
-	check(IsAligned(Num, alignof(FStructSchema)));
+	const int64 HeaderPos = Out.Num(); (void)HeaderPos;
+	check(IsAligned(HeaderPos, alignof(FStructSchema)));
 	WriteData(Out, HeaderData, sizeof(HeaderData));
 	WriteArray(Out, MemberTypes);
 	WriteArray(Out, InnerRangeTypes);
 	WriteAlignedArray(Out, MakeArrayView(MemberNames));
 	WriteAlignedArray(Out, MakeArrayView(InnerSchemas));
 	check(IsAligned(Out.Num(), alignof(FStructSchema)));
-	check(Out.Num() - Num == CalculateSize(*reinterpret_cast<FStructSchema*>(Out.GetData() + Num)));
+	check(Out.Num() - HeaderPos == CalculateSize(*reinterpret_cast<FStructSchema*>(&Out[HeaderPos])));
 }
 
 static bool IsFlatSequence(TConstArrayView<uint64> Constants)
@@ -461,8 +461,8 @@ static void WriteSchema(TArray64<uint8>& Out, const FBuiltEnumSchema& Enum, cons
 	BinaryHeader.Width = Enum.Width;
 	BinaryHeader.Num = IntCastChecked<uint16>(Names.Num());
 
-	const int64 Num = Out.Num(); (void)Num;
-	check(IsAligned(Num, alignof(FEnumSchema)));
+	const int64 HeaderPos = Out.Num(); (void)HeaderPos;
+	check(IsAligned(HeaderPos, alignof(FEnumSchema)));
 	WriteData(Out, &BinaryHeader, sizeof(BinaryHeader));
 	WriteArray(Out, Names);
 	if (BinaryHeader.ExplicitConstants)
@@ -470,7 +470,7 @@ static void WriteSchema(TArray64<uint8>& Out, const FBuiltEnumSchema& Enum, cons
 		WriteEnumConstants(Out, Enum.Width, Enum.Constants);
 	}
 	WriteAlignmentPadding<FEnumSchema>(Out);
-	check(Out.Num() - Num == CalculateSize(*reinterpret_cast<FEnumSchema*>(Out.GetData() + Num)));
+	check(Out.Num() - HeaderPos == CalculateSize(*reinterpret_cast<FEnumSchema*>(&Out[HeaderPos])));
 }
 
 template<typename T>
