PlainProps: Fix clang compile error for fixing incorrect size calcuation (see cl 38328349)

#rnx

diff --git a/Engine/Plugins/Experimental/PlainProps/Source/Private/PlainPropsInternalFormat.h b/Engine/Plugins/Experimental/PlainProps/Source/Private/PlainPropsInternalFormat.h
--- a/Engine/Plugins/Experimental/PlainProps/Source/Private/PlainPropsInternalFormat.h
+++ b/Engine/Plugins/Experimental/PlainProps/Source/Private/PlainPropsInternalFormat.h
@@ -106,7 +106,7 @@ static_assert(offsetof(FStructSchema, Footer) == 15,							"Add binary format ve
 inline uint32 CalculateSize(const FStructSchema& In)
 {
 	static_assert(alignof(FMemberType) == 1);
-	uint32 Out = offsetof(In, Footer) + sizeof(FMemberType) * (uint32(In.NumMembers) + In.NumRangeTypes);
+	uint32 Out = offsetof(FStructSchema, Footer) + sizeof(FMemberType) * (uint32(In.NumMembers) + In.NumRangeTypes);
 	Out = Align(Out + In.NumNames() * sizeof(FMemberId), sizeof(FMemberId));
 	Out = Align(Out + (In.NumInnerSchemas) * sizeof(FSchemaId), sizeof(FSchemaId));
 	return Out;
@@ -132,7 +132,7 @@ static_assert(offsetof(FEnumSchema, Footer) == 12,						"Add binary format versi
 
 inline uint32 CalculateSize(const FEnumSchema& In)
 {
-	return Align(offsetof(In, Footer) + In.Num * sizeof(In.Footer[0]) + In.ExplicitConstants * In.Num * SizeOf(In.Width), alignof(FEnumSchema));
+	return Align(offsetof(FEnumSchema, Footer) + In.Num * sizeof(In.Footer[0]) + In.ExplicitConstants * In.Num * SizeOf(In.Width), alignof(FEnumSchema));
 }
 
 //////////////////////////////////////////////////////////////////////////
