PlainProps alternate fix for deducing sparseness

#rnx

diff --git a/Engine/Plugins/Experimental/PlainProps/Source/Private/PlainPropsBuildSchema.cpp b/Engine/Plugins/Experimental/PlainProps/Source/Private/PlainPropsBuildSchema.cpp
--- a/Engine/Plugins/Experimental/PlainProps/Source/Private/PlainPropsBuildSchema.cpp
+++ b/Engine/Plugins/Experimental/PlainProps/Source/Private/PlainPropsBuildSchema.cpp
@@ -53,9 +53,9 @@ struct FStructSchemaBuilder
 
 	const FStructDeclaration&					Declaration;
 	FSchemasBuilder&							AllSchemas;
-	bool										bMissingMemberNoted = false;
-	FOptionalMemberId*							MemberOrder;
-	FMemberSchema*								NotedSchemas;
+	uint16										MinMembers = 0xFFFF;
+	FOptionalMemberId*							MemberOrder = nullptr;
+	FMemberSchema*								NotedSchemas = nullptr;
 	TBitArray<>									NotedMembers;
 
 	void										NoteMembersRecursively(const FBuiltStruct& Struct);
@@ -171,10 +171,14 @@ FStructSchemaBuilder::FStructSchemaBuilder(const FStructDeclaration& Decl, FSche
 , AllSchemas(Schemas)
 , NotedMembers(false, Declaration.NumMembers + int32(!!Declaration.Super))
 {
+	if (NotedMembers.IsEmpty())
+	{
+		return;
+	}
+
 	MemberOrder = Schemas.GetScratch().AllocateArray<FOptionalMemberId>(NotedMembers.Num());
 	NotedSchemas = Schemas.GetScratch().AllocateArray<FMemberSchema>(NotedMembers.Num());
 
-	check(NotedMembers.Num());
 	MemberOrder[0] = NoId;
 	TConstArrayView<FMemberId> Order = Declaration.GetMemberOrder();
 	FMemory::Memcpy(MemberOrder + int32(!!Declaration.Super), Order.GetData(), Order.NumBytes());
@@ -214,14 +218,13 @@ static void* NoteStructOrEnum(FSchemasBuilder& AllSchemas, bool bStruct, FSchema
 void FStructSchemaBuilder::NoteMembersRecursively(const FBuiltStruct& Struct)
 {
 	check(Declaration.Occupancy != EMemberPresence::RequireAll || Struct.NumMembers == Declaration.NumMembers);
-	bMissingMemberNoted |= Struct.NumMembers < NotedMembers.Num();
+	MinMembers = FMath::Min(MinMembers, Struct.NumMembers);
 	
 	if (Struct.NumMembers == 0)
 	{
 		return;
 	}
 
-
 	const int32 NumNoted = NotedMembers.Num();
 	int32 NoteIdx = 0;
 	for (const FBuiltMember& Member : MakeArrayView(Struct.Members, Struct.NumMembers))
@@ -348,11 +351,11 @@ void FStructSchemaBuilder::NoteRangeRecursively(ERangeSizeType NumType, TConstAr
 
 FBuiltStructSchema FStructSchemaBuilder::Build() const
 {
-	FBuiltStructSchema Out = { Declaration.Type, Declaration.Id, Declaration.Super };
-	Out.bDense = Declaration.Occupancy == EMemberPresence::RequireAll || !bMissingMemberNoted;
+	FBuiltStructSchema Out = { Declaration.Type, Declaration.Id, Declaration.Super, /* dense */ true };
 	
 	if (int32 Num = NotedMembers.CountSetBits())
 	{
+		Out.bDense = Declaration.Occupancy == EMemberPresence::RequireAll || MinMembers == Num;
 		Out.MemberNames.Reserve(Num);
 		Out.MemberSchemas.Reserve(Num);
 		for (int32 NoteIdx = 0, NoteNum = NotedMembers.Num(); NoteIdx < NoteNum; ++NoteIdx)
