PlainProps: Extend "Dense" and "Struct" unit test coverage of deducing sparseness (see cl 38404449).

#rnx

diff --git a/Engine/Plugins/Experimental/PlainProps/Source/Tests/PlainPropsTests.cpp b/Engine/Plugins/Experimental/PlainProps/Source/Tests/PlainPropsTests.cpp
--- a/Engine/Plugins/Experimental/PlainProps/Source/Tests/PlainPropsTests.cpp
+++ b/Engine/Plugins/Experimental/PlainProps/Source/Tests/PlainPropsTests.cpp
@@ -634,12 +634,13 @@ TEST_CASE_NAMED(FPlainPropsReadWriteTest, "System::Core::Serialization::PlainPro
 
 	SECTION("Dense")
 	{
-		static constexpr std::initializer_list<const char*> MemberNames =  {"A", "B", "C"};
+		static constexpr std::initializer_list<const char*> ExplicitMembers = {"A", "B", "C"};
+		static constexpr std::initializer_list<const char*> ImplicitMembers = {"0", "A", "1", "B", "2", "C", "3"};
 
 		TestSerialize([](FTestBatchBuilder& Batch, FScratchAllocator& Scratch)
 		{
-			FStructSchemaId ExplicitId = Batch.DeclareStruct("Test", "ExplicitDense", MemberNames, EMemberPresence::RequireAll);
-			FStructSchemaId ImplicitId = Batch.DeclareStruct("Test", "ImplicitDense", MemberNames, EMemberPresence::AllowSparse);
+			FStructSchemaId ExplicitId = Batch.DeclareStruct("Test", "ExplicitDense", ExplicitMembers, EMemberPresence::RequireAll);
+			FStructSchemaId ImplicitId = Batch.DeclareStruct("Test", "ImplicitDense", ImplicitMembers, EMemberPresence::AllowSparse);
 
 			FMemberBuilder X;
 			X.Add(Batch.NameMember("A"), char8_t('a'));
@@ -657,9 +658,20 @@ TEST_CASE_NAMED(FPlainPropsReadWriteTest, "System::Core::Serialization::PlainPro
 		[](TConstArrayView<FStructView> Objects, const FTestNameReader& Names)
 		{
 			CHECK(Objects.Num() == 2);
+
+			const FStructSchema& ExplicitSchema = Objects[0].Schema.Resolve();
+			const FStructSchema& ImplicitSchema = Objects[1].Schema.Resolve();
+			CHECK(Names[ExplicitSchema.Type.Name] == "ExplicitDense");
+			CHECK(ExplicitSchema.NumMembers == 3);
+			CHECK(ExplicitSchema.IsDense == 1);
+
+			CHECK(Names[ImplicitSchema.Type.Name] == "ImplicitDense");
+			CHECK(ImplicitSchema.NumMembers == 3);
+			CHECK(ImplicitSchema.IsDense == 1);
+
 			FTestMemberReader X(Objects[0]);
 			FTestMemberReader Y(Objects[1]);
-			
+
 			CHECK(Names[X.PeekName()]		== "A");
 			CHECK(X.GrabLeaf().AsChar8()	== 'a');
 			CHECK(Names[X.PeekName()]		== "B");
@@ -678,15 +690,19 @@ TEST_CASE_NAMED(FPlainPropsReadWriteTest, "System::Core::Serialization::PlainPro
 
 	SECTION("Struct")
 	{
-		static constexpr std::initializer_list<const char*> ObjectMembers =  {"L1", "S", "N", "L2"};
-		static constexpr std::initializer_list<const char*> StructMembers =  {"Nested", "Leaf"};
-		static constexpr std::initializer_list<const char*> NestedMembers =  {"I1", "I2"};
+		static constexpr std::initializer_list<const char*> ObjectMembers = {"L1", "S", "N", "L2"};
+		static constexpr std::initializer_list<const char*> StructMembers = {"Nested", "Leaf"};
+		static constexpr std::initializer_list<const char*> NestedMembers = {"I1", "I2"};
+		static constexpr std::initializer_list<const char*> UnusedMembers = {"Unused1", "Unused2"};
+		static constexpr std::initializer_list<const char*> EmptyMembers  = {};
 
 		TestSerialize([](FTestBatchBuilder& Batch, FScratchAllocator& Scratch)
 		{
 			FStructSchemaId ObjectId = Batch.DeclareStruct("Test", "Object", ObjectMembers, EMemberPresence::AllowSparse);
 			FStructSchemaId StructId = Batch.DeclareStruct("Test", "Struct", StructMembers, EMemberPresence::AllowSparse);
 			FStructSchemaId NestedId = Batch.DeclareStruct("Test", "Nested", NestedMembers, EMemberPresence::AllowSparse);
+			FStructSchemaId UnusedId = Batch.DeclareStruct("Test", "Unused", UnusedMembers, EMemberPresence::AllowSparse);
+			FStructSchemaId EmptyId  = Batch.DeclareStruct("Test", "Empty",  EmptyMembers,  EMemberPresence::AllowSparse);
 
 			FMemberBuilder Members;
 			Members.Add(Batch.NameMember("I1"), 100);
@@ -705,10 +721,16 @@ TEST_CASE_NAMED(FPlainPropsReadWriteTest, "System::Core::Serialization::PlainPro
 			Members.Add(Batch.NameMember("L2"), -45.f);
 
 			Batch.AddObject(ObjectId, MoveTemp(Members));
+			Batch.AddObject(UnusedId, {});
+			Batch.AddObject(EmptyId, {});
 		}, 
 		[](TConstArrayView<FStructView> Objects, const FTestNameReader& Names)
 		{
-			CHECK(Objects.Num() == 1);
+			CHECK(Objects.Num() == 3);
+
+			CHECK(Objects[0].Schema.Resolve().IsDense == 1);
+			CHECK(Objects[1].Schema.Resolve().IsDense == 1);
+			CHECK(Objects[2].Schema.Resolve().IsDense == 1);
 
 			FTestMemberReader Object(Objects[0]);
 			CHECK(Object.GrabLeaf().AsFloat() == 123.f);
@@ -722,6 +744,9 @@ TEST_CASE_NAMED(FPlainPropsReadWriteTest, "System::Core::Serialization::PlainPro
 			CHECK(NestedInObject.GrabLeaf().AsS32() == 200);
 
 			CHECK(NestedInStruct.GrabLeaf().AsS32() == 100);
+
+			FTestMemberReader Unused(Objects[1]);  
+			FTestMemberReader Empty(Objects[2]);
 		});
 	}
 
