PlainProps initial numeral implementation and CoreUObject numeral bindings for static arrays

Long-term FMemberId might be a numeral-or-name encoding, FNumeralMemberBuilder/Reader APIs, a more compact schema representation and disk format, etc. For now, each ID system generates its own numeral member IDs from strings.

#rnx

diff --git a/Engine/Plugins/Experimental/PlainProps/Source/Private/PlainPropsBuild.cpp b/Engine/Plugins/Experimental/PlainProps/Source/Private/PlainPropsBuild.cpp
--- a/Engine/Plugins/Experimental/PlainProps/Source/Private/PlainPropsBuild.cpp
+++ b/Engine/Plugins/Experimental/PlainProps/Source/Private/PlainPropsBuild.cpp
@@ -130,7 +130,7 @@ void FMemberBuilder::BuildSuperStruct(FScratchAllocator& Scratch, const FStructD
 	
 	FBuiltStruct* OnlyMember = BuildAndReset(Scratch, Super, Debug);
 	Members.Emplace(FBuiltMember::MakeSuper(Super.Id, MoveTemp(OnlyMember)));
-	check(Members[0].Schema.Type.AsStruct().IsSuper);
+	check(IsSuper(Members[0].Schema.Type));
 }
 
 FBuiltStruct* FMemberBuilder::BuildAndReset(FScratchAllocator& Scratch, const FStructDeclaration& Declared, const FDebugIds& Debug)
@@ -141,20 +141,18 @@ FBuiltStruct* FMemberBuilder::BuildAndReset(FScratchAllocator& Scratch, const FS
 	// Verify members were added in declared order
 	if (int32 Num = Members.Num())
 	{
-		if (Declared.bNamedMembers)
+		const bool HasSuper = IsSuper(Members[0].Schema.Type);
+		int32 OrderIdx = 0;
+		TConstArrayView<FMemberId> Order = Declared.GetMemberOrder();
+		int32 SkipSuper = Declared.Super && HasSuper;
+		for (FBuiltMember* It = Members.GetData() + SkipSuper, *End = Members.GetData() + Num; It != End; ++It)
 		{
-			int32 OrderIdx = 0;
-			TConstArrayView<FMemberId> Order = Declared.GetMemberOrder();
-			int32 SkipSuper = Declared.Super && IsSuper(Members[0].Schema.Type);
-			for (FBuiltMember* It = Members.GetData() + SkipSuper, *End = Members.GetData() + Num; It != End; ++It)
-			{
-				FBuiltMember& Member = *It;
-				for (; OrderIdx < Order.Num() && Order[OrderIdx] != Member.Name; ++OrderIdx)
-				{}	
-				checkf(OrderIdx < Order.Num(), TEXT("Member '%s' in '%s' %s"), *Debug.Print(Member.Name), *Debug.Print(Declared.Type),
-						Order.Contains(Member.Name) ? TEXT("appeared in non-declared order") : TEXT("is undeclared"));
-				++OrderIdx;
-			}
+			FBuiltMember& Member = *It;
+			for (; OrderIdx < Order.Num() && Order[OrderIdx] != Member.Name; ++OrderIdx)
+			{}	
+			checkf(OrderIdx < Order.Num(), TEXT("Member '%s' in '%s' %s"), *Debug.Print(Member.Name), *Debug.Print(Declared.Type),
+					Order.Contains(Member.Name) ? TEXT("appeared in non-declared order") : TEXT("is undeclared"));
+			++OrderIdx;
 		}
 	}
 #endif
@@ -275,4 +273,4 @@ FTypedRange FNestedRangeBuilder::BuildAndReset(FScratchAllocator& Scratch, ERang
 	return { MakeNestedRangeSchema(Scratch, SizeType, Schema), Out };
 }
 
-} // namespace PlainProps
\ No newline at end of file
+} // namespace PlainProps
diff --git a/Engine/Plugins/Experimental/PlainProps/Source/Private/PlainPropsBuildSchema.cpp b/Engine/Plugins/Experimental/PlainProps/Source/Private/PlainPropsBuildSchema.cpp
--- a/Engine/Plugins/Experimental/PlainProps/Source/Private/PlainPropsBuildSchema.cpp
+++ b/Engine/Plugins/Experimental/PlainProps/Source/Private/PlainPropsBuildSchema.cpp
@@ -171,17 +171,15 @@ FStructSchemaBuilder::FStructSchemaBuilder(const FStructDeclaration& Decl, FSche
 , AllSchemas(Schemas)
 , NotedMembers(false, Declaration.NumMembers + int32(!!Declaration.Super))
 {
-	if (NotedMembers.IsEmpty())
+	if (int32 Num = NotedMembers.Num())
 	{
-		return;
-	}
+		MemberOrder = Schemas.GetScratch().AllocateArray<FOptionalMemberId>(Num);
+		NotedSchemas = Schemas.GetScratch().AllocateArray<FMemberSchema>(Num);
 
-	MemberOrder = Schemas.GetScratch().AllocateArray<FOptionalMemberId>(NotedMembers.Num());
-	NotedSchemas = Schemas.GetScratch().AllocateArray<FMemberSchema>(NotedMembers.Num());
-
-	MemberOrder[0] = NoId;
-	TConstArrayView<FMemberId> Order = Declaration.GetMemberOrder();
-	FMemory::Memcpy(MemberOrder + int32(!!Declaration.Super), Order.GetData(), Order.NumBytes());
+		MemberOrder[0] = NoId; // In case unnamed super exists
+		TConstArrayView<FMemberId> Order = Declaration.GetMemberOrder();
+		FMemory::Memcpy(MemberOrder + int32(!!Declaration.Super), Order.GetData(), Order.NumBytes());
+	}
 }
 
 static bool RequiresDynamicStructSchema(const FMemberSchema& A, const FMemberSchema& B)
diff --git a/Engine/Plugins/Experimental/PlainProps/Source/Private/PlainPropsDeclare.cpp b/Engine/Plugins/Experimental/PlainProps/Source/Private/PlainPropsDeclare.cpp
--- a/Engine/Plugins/Experimental/PlainProps/Source/Private/PlainPropsDeclare.cpp
+++ b/Engine/Plugins/Experimental/PlainProps/Source/Private/PlainPropsDeclare.cpp
@@ -118,7 +118,7 @@ static const FStructDeclaration& Declare(TArray<TUniquePtr<FStructDeclaration>>&
 	}
 	else
 	{
-		FStructDeclaration Header{1, DeclId, Type, Super, Occupancy, !MemberOrder.IsEmpty(), NumMembers};
+		FStructDeclaration Header{1, DeclId, Type, Super, Occupancy, NumMembers};
 		void* Data = FMemory::Malloc(sizeof(FStructDeclaration) + MemberOrder.Num() * MemberOrder.GetTypeSize());
 		Ptr.Reset(new (Data) FStructDeclaration(Header));
 		CopyItems(Ptr->MemberOrder, MemberOrder);
@@ -127,15 +127,14 @@ static const FStructDeclaration& Declare(TArray<TUniquePtr<FStructDeclaration>>&
 	return *Ptr;
 }
 
-
 const FStructDeclaration& FDeclarations::DeclareStruct(FStructSchemaId DeclId, FTypeId Type, TConstArrayView<FMemberId> MemberOrder, EMemberPresence Occupancy, FOptionalStructSchemaId Super)
 {
 	return Declare(/* out */ DeclaredStructs, DeclId, Type, MemberOrder, Occupancy, Super, IntCastChecked<uint16>(MemberOrder.Num()));
 }
 
-const FStructDeclaration& FDeclarations::DeclareNumeralStruct(FStructSchemaId DeclId, FTypeId Type, uint16 NumMembers, EMemberPresence Occupancy)
+const FStructDeclaration& FDeclarations::DeclareNumeralStruct(FStructSchemaId DeclId, FTypeId Type, TConstArrayView<FMemberId> Numerals, EMemberPresence Occupancy)
 {
-	return Declare(/* out */ DeclaredStructs, DeclId, Type, {}, Occupancy, NoId, NumMembers);
+	return Declare(/* out */ DeclaredStructs, DeclId, Type, Numerals, Occupancy, NoId, IntCastChecked<uint16>(Numerals.Num()));
 }
 
 const FEnumDeclaration& FDeclarations::DeclareEnum(FEnumSchemaId Id, FTypeId Type, EEnumMode Mode, TConstArrayView<FEnumerator> Enumerators, EEnumAliases Policy)
diff --git a/Engine/Plugins/Experimental/PlainProps/Source/Private/PlainPropsSave.cpp b/Engine/Plugins/Experimental/PlainProps/Source/Private/PlainPropsSave.cpp
--- a/Engine/Plugins/Experimental/PlainProps/Source/Private/PlainPropsSave.cpp
+++ b/Engine/Plugins/Experimental/PlainProps/Source/Private/PlainPropsSave.cpp
@@ -533,34 +533,16 @@ static void SaveMemberDelta(FMemberBuilder& Out, const void* Struct, const void*
 static const FStructDeclaration& SaveSchemaBoundStructDelta(FMemberBuilder& Out, const void* Struct, const void* Default, const FSchemaBinding& Schema, const FSaveContext& Ctx)
 {
 	const FStructDeclaration& Declaration = Ctx.Declarations.Get(Schema.DeclId);
-	if (Declaration.bNamedMembers)
-	{
-		TConstArrayView<FMemberId> MemberOrder = Declaration.GetMemberOrder();
+	TConstArrayView<FMemberId> MemberOrder = Declaration.GetMemberOrder();
 
-		for (FMemberVisitor It(Schema); It.HasMore(); )
-		{
-			FMemberId Name = MemberOrder[It.GetIndex()];
-			switch (It.PeekKind())
-			{
-				case EMemberKind::Leaf:		SaveMemberDelta(Out, Struct, Default, Name, Ctx, It.GrabLeaf());	break;
-				case EMemberKind::Range:	SaveMemberDelta(Out, Struct, Default, Name, Ctx, It.GrabRange());	break;
-				case EMemberKind::Struct:	SaveMemberDelta(Out, Struct, Default, Name, Ctx, It.GrabStruct());	break;
-			}
-		}
-	}
-	else // numeral struct
+	for (FMemberVisitor It(Schema); It.HasMore(); )
 	{
-		for (FMemberVisitor It(Schema); It.HasMore(); )
+		FMemberId Name = MemberOrder[It.GetIndex()];
+		switch (It.PeekKind())
 		{
-			// Todo: Make real numeral
-			FMemberId Dummy = { 0 };
-			switch (It.PeekKind())
-			{
-				case EMemberKind::Leaf:		SaveMemberDelta(Out, Struct, Default, Dummy, Ctx, It.GrabLeaf());	break;
-				case EMemberKind::Range:	SaveMemberDelta(Out, Struct, Default, Dummy, Ctx, It.GrabRange());	break;
-				case EMemberKind::Struct:	SaveMemberDelta(Out, Struct, Default, Dummy, Ctx, It.GrabStruct());	break;
-			}
-			++Dummy.Id.Idx;
+			case EMemberKind::Leaf:		SaveMemberDelta(Out, Struct, Default, Name, Ctx, It.GrabLeaf());	break;
+			case EMemberKind::Range:	SaveMemberDelta(Out, Struct, Default, Name, Ctx, It.GrabRange());	break;
+			case EMemberKind::Struct:	SaveMemberDelta(Out, Struct, Default, Name, Ctx, It.GrabStruct());	break;
 		}
 	}
 
diff --git a/Engine/Plugins/Experimental/PlainProps/Source/Public/PlainPropsDeclare.h b/Engine/Plugins/Experimental/PlainProps/Source/Public/PlainPropsDeclare.h
--- a/Engine/Plugins/Experimental/PlainProps/Source/Public/PlainPropsDeclare.h
+++ b/Engine/Plugins/Experimental/PlainProps/Source/Public/PlainPropsDeclare.h
@@ -41,13 +41,11 @@ struct FStructDeclaration
 	FTypeId					Type;			// Could be removed, might allow declaration dedup among templated types
 	FOptionalStructSchemaId	Super;
 	EMemberPresence			Occupancy;
-	bool					bNamedMembers;
 	uint16					NumMembers;
-	FMemberId				MemberOrder[0]; // Only if bNamedMembers
+	FMemberId				MemberOrder[0];
 
 	TConstArrayView<FMemberId> GetMemberOrder() const
 	{
-		check(bNamedMembers);
 		return MakeArrayView(MemberOrder, NumMembers);
 	}
 
@@ -67,7 +65,7 @@ public:
 	PLAINPROPS_API const FEnumDeclaration&			DeclareEnum(FEnumSchemaId DeclId, FTypeId Type, EEnumMode Mode, TConstArrayView<FEnumerator> Enumerators, EEnumAliases Policy);
 	// Declare struct with ref count 1 or increment it and check that previous declaration matches
 	PLAINPROPS_API const FStructDeclaration&		DeclareStruct(FStructSchemaId DeclId, FTypeId Type, TConstArrayView<FMemberId> MemberOrder, EMemberPresence Occupancy, FOptionalStructSchemaId Super = {});
-	PLAINPROPS_API const FStructDeclaration&		DeclareNumeralStruct(FStructSchemaId DeclId, FTypeId Type, uint16 NumMembers, EMemberPresence Occupancy);
+	PLAINPROPS_API const FStructDeclaration&		DeclareNumeralStruct(FStructSchemaId DeclId, FTypeId Type, TConstArrayView<FMemberId> Numerals, EMemberPresence Occupancy);
 	
 	void											DropEnum(FEnumSchemaId DeclId)			{ Check(DeclId); DeclaredEnums[DeclId.Idx].Reset(); }
 	PLAINPROPS_API void								DropStructRef(FStructSchemaId DeclId);
diff --git a/Engine/Plugins/Experimental/PlainProps/Source/Public/PlainPropsUeCoreBindings.h b/Engine/Plugins/Experimental/PlainProps/Source/Public/PlainPropsUeCoreBindings.h
--- a/Engine/Plugins/Experimental/PlainProps/Source/Public/PlainPropsUeCoreBindings.h
+++ b/Engine/Plugins/Experimental/PlainProps/Source/Public/PlainPropsUeCoreBindings.h
@@ -662,14 +662,14 @@ struct TBaseDeltaBinding : ICustomBinding, FBaseDeltaBinding
 				
 		while (Members.HasMore())
 		{
-			if (Members.PeekName() == AddId())
+			if (Members.PeekNameUnchecked() == AddId())
 			{
 				ApplyItems<EOp::Add, ElemType>(Dst, Members.GrabRange(), Elems);
 				check(!Members.HasMore());
 				break;
 			}
 
-			check(Members.PeekName() == DelId());
+			check(Members.PeekNameUnchecked() == DelId());
 			ApplyItems<EOp::Del, KeyType>(Dst, Members.GrabRange(), Keys);
 		}
 	}
diff --git a/Engine/Plugins/Experimental/PlainPropsUObject/Source/Private/PlainPropsUObjectBindings.cpp b/Engine/Plugins/Experimental/PlainPropsUObject/Source/Private/PlainPropsUObjectBindings.cpp
--- a/Engine/Plugins/Experimental/PlainPropsUObject/Source/Private/PlainPropsUObjectBindings.cpp
+++ b/Engine/Plugins/Experimental/PlainPropsUObject/Source/Private/PlainPropsUObjectBindings.cpp
@@ -191,6 +191,24 @@ const void* FDefaultStructs::Get(FStructSchemaId Id)
 
 ////////////////////////////////////////////////////////////////////////////////////////////////
 
+FMemberId FNumeralGenerator::Grow(int32 Max)
+{
+	check(Max > Cache.Num());
+	int32 OldNum = Cache.Num();
+	Cache.SetNumUninitialized(Max + 1);
+
+	FName Numeral("_");
+	for (int32 Idx = OldNum; Idx <= Max; ++Idx)
+	{
+		Numeral.SetNumber(NAME_EXTERNAL_TO_INTERNAL(Idx));
+		Cache[Idx] = GUE.Names.NameMember(Numeral);
+	}
+
+	return Cache[Max];
+}
+
+////////////////////////////////////////////////////////////////////////////////////////////////
+
 FCommonScopeIds::FCommonScopeIds(TIdIndexer<FName>& Names)
 : Core(						Names.MakeScope("/Script/Core"))
 , CoreUObject(				Names.MakeScope("/Script/CoreUObject"))
@@ -2805,7 +2823,8 @@ private:
 			FStructSchemaId StaticArrayId = GUE.Names.IndexStruct(StaticArrayType);
 
 			// Todo: Ownership
-			GUE.Types.DeclareNumeralStruct(StaticArrayId, StaticArrayType, ArrayDim, Occupancy);
+			TConstArrayView<FMemberId> Numerals = GUE.Numerals.MakeRange(IntCastChecked<uint16>(ArrayDim));
+			GUE.Types.DeclareNumeralStruct(StaticArrayId, StaticArrayType, Numerals, Occupancy);
 
 			TArray<FMemberBinding, TInlineAllocator<64>> Elements;
 			Elements.Init(Out, ArrayDim);
diff --git a/Engine/Plugins/Experimental/PlainPropsUObject/Source/Public/PlainPropsUObjectRuntime.h b/Engine/Plugins/Experimental/PlainPropsUObject/Source/Public/PlainPropsUObjectRuntime.h
--- a/Engine/Plugins/Experimental/PlainPropsUObject/Source/Public/PlainPropsUObjectRuntime.h
+++ b/Engine/Plugins/Experimental/PlainPropsUObject/Source/Public/PlainPropsUObjectRuntime.h
@@ -38,6 +38,26 @@ private:
 	void									ReserveFlags(uint32 Idx);
 };
 
+// Temporary implementation for numeral structs
+class FNumeralGenerator
+{
+	TArray<FMemberId>			Cache; // 0, 1, 2, ..
+
+	FMemberId					Grow(int32 Max);
+
+public:
+	inline FMemberId			Make(uint16 Numeral)
+	{
+		return int32(Numeral) < Cache.Num() ? Cache[Numeral] : Grow(Numeral);
+	}
+	
+	TConstArrayView<FMemberId>	MakeRange(uint16 Numerals)
+	{
+		Make(Numerals);
+		return MakeArrayView(Cache.GetData(), Numerals);
+	}
+};
+
 struct FCommonScopeIds
 {
 	explicit FCommonScopeIds(TIdIndexer<FName>& Names);
@@ -106,6 +126,7 @@ struct FGlobals
 	FCommonStructIds			Structs;
 	FCommonTypenameIds			Typenames;
 	FCommonMemberIds			Members;
+	FNumeralGenerator			Numerals;
 };
 PLAINPROPSUOBJECT_API extern FGlobals GUE;
 
