PlainProps fix debug build

#rnx

diff --git a/Engine/Plugins/Experimental/PlainProps/Source/Public/PlainPropsLoadMember.h b/Engine/Plugins/Experimental/PlainProps/Source/Public/PlainPropsLoadMember.h
--- a/Engine/Plugins/Experimental/PlainProps/Source/Public/PlainPropsLoadMember.h
+++ b/Engine/Plugins/Experimental/PlainProps/Source/Public/PlainPropsLoadMember.h
@@ -34,21 +34,10 @@ PLAINPROPS_API void LoadSoleStruct(void* Dst, FStructLoadView Src);
 
 // Load single leaf member faster than FMemberLoader(Src).GrabLeaf().As<T>()
 template<LeafType T>
-inline void LoadSole(void* Dst, FStructLoadView Src)
-{
-	Src.Values.CheckSize(sizeof(T));
-	checkSlow(ReflectLeaf<T> == SizeOf(FMemberLoader(Src).GrabLeaf().Leaf));
-	FMemory::Memcpy(Dst, Src.Values.Peek(), sizeof(T));
-}
+inline void LoadSole(void* Dst, FStructLoadView Src);
 
 template<>
-inline void LoadSole<bool>(void* Dst, FStructLoadView Src)
-{
-	Src.Values.CheckSize(sizeof(bool));
-	checkSlow(ReflectLeaf<T> == SizeOf(FMemberLoader(Src).GrabLeaf().Leaf));
-	FBitCacheReader Bits;
-	*static_cast<bool*>(Dst) = Bits.GrabNext(Src.Values);
-}
+inline void LoadSole<bool>(void* Dst, FStructLoadView Src);
 
 template<LeafType T>
 inline T LoadSole(FStructLoadView Src)
@@ -221,4 +210,23 @@ private:
 	const FLoadBatch&		Batch;
 };
 
+//////////////////////////////////////////////////////////////////////////
+
+template<LeafType T>
+void LoadSole(void* Dst, FStructLoadView Src)
+{
+	Src.Values.CheckSize(sizeof(T));
+	checkSlow(ReflectLeaf<T> == FMemberLoader(Src).GrabLeaf().Leaf);
+	FMemory::Memcpy(Dst, Src.Values.Peek(), sizeof(T));
+}
+
+template<>
+void LoadSole<bool>(void* Dst, FStructLoadView Src)
+{
+	Src.Values.CheckSize(sizeof(bool));
+	checkSlow(ReflectLeaf<bool> == FMemberLoader(Src).GrabLeaf().Leaf);
+	FBitCacheReader Bits;
+	*static_cast<bool*>(Dst) = Bits.GrabNext(Src.Values);
+}
+
 } // namespace PlainProps
\ No newline at end of file
