PlainProps load related fixes, can now roundtrip all UObjects during ShooterGame init

#rnx

diff --git a/Engine/Plugins/Experimental/PlainProps/Source/Private/PlainPropsBind.cpp b/Engine/Plugins/Experimental/PlainProps/Source/Private/PlainPropsBind.cpp
--- a/Engine/Plugins/Experimental/PlainProps/Source/Private/PlainPropsBind.cpp
+++ b/Engine/Plugins/Experimental/PlainProps/Source/Private/PlainPropsBind.cpp
@@ -442,7 +442,7 @@ FStructSchemaId FStructBindIds::GetDeclId(FStructSchemaId BindId) const
 
 //////////////////////////////////////////////////////////////////////////
 
-TArray<FStructSchemaId> IndexInMemoryNames(const FSchemaBatch& Schemas, FIdIndexerBase& Indexer)
+TArray<FStructSchemaId> IndexRuntimeIds(const FSchemaBatch& Schemas, FIdIndexerBase& Indexer)
 {
 	const uint8* Base = reinterpret_cast<const uint8*>(&Schemas);
 	const uint32* Offsets = Schemas.SchemaOffsets;
diff --git a/Engine/Plugins/Experimental/PlainProps/Source/Private/PlainPropsBuildSchema.cpp b/Engine/Plugins/Experimental/PlainProps/Source/Private/PlainPropsBuildSchema.cpp
--- a/Engine/Plugins/Experimental/PlainProps/Source/Private/PlainPropsBuildSchema.cpp
+++ b/Engine/Plugins/Experimental/PlainProps/Source/Private/PlainPropsBuildSchema.cpp
@@ -49,11 +49,12 @@ static FString PrintMemberSchema(const FDebugIds& Debug, FMemberSchema Schema)
 
 struct FStructSchemaBuilder
 {
-	FStructSchemaBuilder(const FStructDeclaration& InDecl, FSchemasBuilder& InSchemas);
+	FStructSchemaBuilder(FStructSchemaId InId, const FStructDeclaration& InDecl, FSchemasBuilder& InSchemas);
 
 	const FStructDeclaration&					Declaration;
 	FSchemasBuilder&							AllSchemas;
 	uint16										MinMembers = 0xFFFF;
+	FStructSchemaId								Id;
 	FOptionalMemberId*							MemberOrder = nullptr;
 	FMemberSchema*								NotedSchemas = nullptr;
 	TBitArray<>									NotedMembers;
@@ -77,40 +78,41 @@ struct FEnumSchemaBuilder
 
 //////////////////////////////////////////////////////////////////////////
 
-FSchemasBuilder::FSchemasBuilder(const FDeclarations& Types, const IStructBindIds& InBindIds, FScratchAllocator& InScratch)
-: FSchemasBuilder(Types.GetStructs(), Types.GetEnums(), InBindIds, Types.GetDebug(), InScratch)
+FSchemasBuilder::FSchemasBuilder(const FDeclarations& Types, const IStructBindIds& InBindIds, FScratchAllocator& InScratch, ESchemaFormat InFormat)
+: FSchemasBuilder(Types.GetStructs(), Types.GetEnums(), InBindIds, Types.GetDebug(), InScratch, InFormat)
 {}
 
-FSchemasBuilder::FSchemasBuilder(FStructDeclarations InStructs, FEnumDeclarations InEnums, const IStructBindIds& InBindIds, const FDebugIds& InDebug, FScratchAllocator& InScratch)
+FSchemasBuilder::FSchemasBuilder(FStructDeclarations InStructs, FEnumDeclarations InEnums, const IStructBindIds& InBindIds, const FDebugIds& InDebug, FScratchAllocator& InScratch, ESchemaFormat InFormat)
 : DeclaredStructs(InStructs)
 , DeclaredEnums(InEnums)
 , BindIds(InBindIds)
+, Format(InFormat)
 , Scratch(InScratch)
 , Debug(InDebug)
-{
-	StructIndices.Init(INDEX_NONE, InStructs.Num());
-	EnumIndices.Init(INDEX_NONE, InEnums.Num());
-}
+{}
 
 FSchemasBuilder::~FSchemasBuilder() {}
 
-template<class T, typename ...Ts>
-FORCEINLINE T& GetOrEmplace(int32& Index, TPagedArray<T, 4096>& Things, Ts&&... EmplaceArgs)
+template<class IdType, class BuilderType, typename ...Ts>
+FORCEINLINE BuilderType& GetOrEmplace(TSet<IdType>& Indexer, IdType Id, TPagedArray<BuilderType, 4096>& Builders, Ts&&... EmplaceArgs)
 {
-	if (Index == INDEX_NONE)
+	FSetElementId Idx = Indexer.FindId(Id);
+	if (Idx.IsValidId())
 	{
-		Index = Things.Num();
-		Things.Emplace(Forward<Ts>(EmplaceArgs)...);
+		return Builders[Idx.AsInteger()];
 	}
 
-	return Things[Index];
+	Idx = Indexer.Emplace(Id);
+	check(Idx.AsInteger() == Builders.Num());
+	Builders.Emplace(Forward<Ts>(EmplaceArgs)...);
+	return Builders.Last();
 }
 	
 FORCEINLINE FEnumSchemaBuilder&	FSchemasBuilder::NoteEnum(FEnumSchemaId Id)
 {
 	checkf(!bBuilt, TEXT("Noted new members after building"));
 	checkf(Id.Idx < uint32(DeclaredEnums.Num()) && DeclaredEnums[Id.Idx], TEXT("Undeclared enum '%s' noted"), *Debug.Print(Id));
-	return GetOrEmplace(EnumIndices[Id.Idx], Enums, *DeclaredEnums[Id.Idx], Id);
+	return GetOrEmplace(EnumIndices, Id, Enums, *DeclaredEnums[Id.Idx], Id);
 }
 
 FORCEINLINE FStructSchemaBuilder& FSchemasBuilder::NoteStruct(FStructSchemaId BindId)
@@ -121,7 +123,8 @@ FORCEINLINE FStructSchemaBuilder& FSchemasBuilder::NoteStruct(FStructSchemaId Bi
 							? BindId : BindIds.GetDeclId(BindId);
 
 	checkf(DeclId.Idx < uint32(DeclaredStructs.Num()) && DeclaredStructs[DeclId.Idx], TEXT("Undeclared struct '%s' noted"), *Debug.Print(DeclId));
-	return GetOrEmplace(StructIndices[DeclId.Idx], Structs, *DeclaredStructs[DeclId.Idx], *this);
+	FStructSchemaId NoteId = (Format == ESchemaFormat::StableNames) ? DeclId : BindId;
+	return GetOrEmplace(StructIndices, NoteId, Structs, NoteId, *DeclaredStructs[DeclId.Idx], *this);
 }
 
 void FSchemasBuilder::NoteStructAndMembers(FStructSchemaId BindId, const FBuiltStruct& Struct)
@@ -158,17 +161,18 @@ void FSchemasBuilder::NoteInheritanceChains()
 	{
 		for (FOptionalStructSchemaId Super = Structs[Idx].Declaration.Super; Super; Super = DeclaredStructs[Super.Get().Idx]->Super)
 		{
-			uint32 SuperIdx = Super.Get().Idx;
-			GetOrEmplace(StructIndices[SuperIdx], Structs, *DeclaredStructs[SuperIdx], *this);
+			FStructSchemaId Id = Super.Get();
+			GetOrEmplace(StructIndices, Id, Structs, Super.Get(), *DeclaredStructs[Id.Idx], *this);
 		}
 	}
 }
 
 //////////////////////////////////////////////////////////////////////////
 
-FStructSchemaBuilder::FStructSchemaBuilder(const FStructDeclaration& Decl, FSchemasBuilder& Schemas)
+FStructSchemaBuilder::FStructSchemaBuilder(FStructSchemaId InId, const FStructDeclaration& Decl, FSchemasBuilder& Schemas)
 : Declaration(Decl)
 , AllSchemas(Schemas)
+, Id(InId)
 , NotedMembers(false, Declaration.NumMembers + int32(!!Declaration.Super))
 {
 	if (int32 Num = NotedMembers.Num())
@@ -350,7 +354,8 @@ void FStructSchemaBuilder::NoteRangeRecursively(ERangeSizeType NumType, TConstAr
 
 FBuiltStructSchema FStructSchemaBuilder::Build() const
 {
-	FBuiltStructSchema Out = { Declaration.Type, Declaration.Id, Declaration.Super, /* dense */ true };
+	FTypeId Type = Id == Declaration.Id ? Declaration.Type : AllSchemas.GetDebug().Resolve(Id);
+	FBuiltStructSchema Out = { Type, Id, Declaration.Super, /* dense */ true };
 	
 	if (int32 Num = NotedMembers.CountSetBits())
 	{
@@ -448,4 +453,18 @@ void FEnumSchemaBuilder::NoteEmpty(ELeafWidth Width)
 	NotedWidth = Width;
 }
 
+//////////////////////////////////////////////////////////////////////////
+
+TArray<FStructSchemaId>	ExtractRuntimeIds(const FBuiltSchemas& In)
+{
+	TArray<FStructSchemaId>	Out;
+	Out.SetNumUninitialized(In.Structs.Num());
+	const FBuiltStructSchema* InIt = In.Structs.GetData();
+	for (FStructSchemaId& Id : Out)
+	{
+		Id = (*InIt++).Id;
+	}
+	return Out;
+}
+
 } // namespace PlainProps
diff --git a/Engine/Plugins/Experimental/PlainProps/Source/Private/PlainPropsIndex.cpp b/Engine/Plugins/Experimental/PlainProps/Source/Private/PlainPropsIndex.cpp
--- a/Engine/Plugins/Experimental/PlainProps/Source/Private/PlainPropsIndex.cpp
+++ b/Engine/Plugins/Experimental/PlainProps/Source/Private/PlainPropsIndex.cpp
@@ -164,4 +164,16 @@ FStructSchemaId	FIdIndexerBase::IndexStruct(FTypeId Type)
 	return { IntCastChecked<uint32>(Structs.Add(Type).AsInteger()) };
 }
 
+FTypeId FIdIndexerBase::Resolve(FEnumSchemaId Id) const
+{
+	check(Id.Idx < uint32(Enums.Num()));
+	return Enums.Get(FSetElementId::FromInteger(Id.Idx));
+}
+
+FTypeId FIdIndexerBase::Resolve(FStructSchemaId Id) const
+{
+	check(Id.Idx < uint32(Structs.Num()));
+	return Structs.Get(FSetElementId::FromInteger(Id.Idx));
+}
+
 } // namespace PlainProps
\ No newline at end of file
diff --git a/Engine/Plugins/Experimental/PlainProps/Source/Private/PlainPropsLoad.cpp b/Engine/Plugins/Experimental/PlainProps/Source/Private/PlainPropsLoad.cpp
--- a/Engine/Plugins/Experimental/PlainProps/Source/Private/PlainPropsLoad.cpp
+++ b/Engine/Plugins/Experimental/PlainProps/Source/Private/PlainPropsLoad.cpp
@@ -210,6 +210,13 @@ private:
 	}
 };
 
+struct FLoadIdBinderDummy
+{
+	static constexpr FStructSchemaId	BindLoadId(FStructSchemaId ReadId, FStructSchemaId BindId) { return ReadId; }
+	static constexpr int32				NumMappings() { return 0; }
+	static inline FLoadIdMapping		GetMapping(int32 Idx) { unimplemented(); return {}; }
+};
+
 ////////////////////////////////////////////////////////////////////////////
 
 struct FLoadBatch
@@ -249,7 +256,8 @@ static void CopyMemberBinding(FLeafMemberBinding Binding, /* in-out */ const FSc
 
 // TODO: Make all below free functions FLoadPlanner member functions and make LoadIds a member variable
 
-static void CopyMemberBinding(FStructMemberBinding Binding, /* in-out */ const FSchemaId*& InnerSchemaIt, FLoadIdBinder& LoadIds, FMemberBinder& Out)
+template<class LoadIdBinder>
+void CopyMemberBinding(FStructMemberBinding Binding, /* in-out */ const FSchemaId*& InnerSchemaIt, LoadIdBinder& LoadIds, FMemberBinder& Out)
 {
 	FStructSchemaId LoadId = LoadIds.BindLoadId(static_cast<FStructSchemaId>(*InnerSchemaIt), Binding.Id);
 	Out.AddMember(FMemberBindType(Binding.Type), static_cast<uint32>(Binding.Offset));
@@ -257,7 +265,8 @@ static void CopyMemberBinding(FStructMemberBinding Binding, /* in-out */ const F
 	++InnerSchemaIt;
 }
 
-static void CopyMemberBinding(FRangeMemberBinding Binding, /* in-out */ const FSchemaId*& InnerSchemaIt, FLoadIdBinder& LoadIds, FMemberBinder& Out)
+template<class LoadIdBinder>
+void CopyMemberBinding(FRangeMemberBinding Binding, /* in-out */ const FSchemaId*& InnerSchemaIt, LoadIdBinder& LoadIds, FMemberBinder& Out)
 {
 	FMemberBindType InnermostType = Binding.InnerTypes[Binding.NumRanges - 1];
 	Out.AddRange(MakeArrayView(Binding.RangeBindings, Binding.NumRanges), InnermostType, static_cast<uint32>(Binding.Offset));
@@ -273,7 +282,8 @@ static void CopyMemberBinding(FRangeMemberBinding Binding, /* in-out */ const FS
 	}
 }
 
-static void CopyMemberBinding(/* in-out */ FMemberVisitor& BindIt, /* in-out */ const FSchemaId*& InnerSchemaIt, FLoadIdBinder& LoadIds, FMemberBinder& Out)
+template<class LoadIdBinder>
+void CopyMemberBinding(/* in-out */ FMemberVisitor& BindIt, /* in-out */ const FSchemaId*& InnerSchemaIt, LoadIdBinder& LoadIds, FMemberBinder& Out)
 {
 	switch (BindIt.PeekKind())
 	{
@@ -284,7 +294,8 @@ static void CopyMemberBinding(/* in-out */ FMemberVisitor& BindIt, /* in-out */
 	}
 }
 
-static void CreateSubsetBindingWithoutEnumIds(const FStructSchema& From, const FSchemaBinding& To, TConstArrayView<FMemberId> ToNames, uint16 NumEnums, FLoadIdBinder& LoadIds, SubsetByteArray& Out)
+template<class LoadIdBinder>
+static void CreateSubsetBindingWithoutEnumIds(const FStructSchema& From, const FSchemaBinding& To, TConstArrayView<FMemberId> ToNames, uint16 NumEnums, LoadIdBinder& LoadIds, SubsetByteArray& Out)
 {
 	check(To.NumMembers == ToNames.Num());
 	check(To.NumMembers >= From.NumMembers);
@@ -293,14 +304,15 @@ static void CreateSubsetBindingWithoutEnumIds(const FStructSchema& From, const F
 	int32 OutPos = Out.Num();
 	
 	// Allocate and init header 
-	FSchemaBinding Header = { To.DeclId, From.NumMembers, From.NumInnerSchemas - NumEnums, From.NumRangeTypes };
+	bool bSkipSuperSchema = From.Inheritance == ESuper::Unused;
+	FSchemaBinding Header = { To.DeclId, From.NumMembers, From.NumInnerSchemas - NumEnums - bSkipSuperSchema, From.NumRangeTypes };
 	Out.AddUninitialized(Header.CalculateSize());
 	FSchemaBinding* Schema = new (&Out[OutPos]) FSchemaBinding {Header};
 	
 	// Copy subset of member bindings
 	FMemberVisitor ToIt(To);
 	FMemberBinder Footer(*Schema);
-	const FSchemaId* InnerSchemaIt = From.GetInnerSchemas();
+	const FSchemaId* InnerSchemaIt = From.GetInnerSchemas() + bSkipSuperSchema;
 	for (FMemberId FromName : From.GetMemberNames())
 	{
 		while (FromName != ToNames[ToIt.GetIndex()])
@@ -313,7 +325,8 @@ static void CreateSubsetBindingWithoutEnumIds(const FStructSchema& From, const F
 	check(InnerSchemaIt == From.GetInnerSchemas() + From.NumInnerSchemas);
 }
 
-static void CloneBindingWithReplacedStructIds(const FSchemaId* FromIds, const FSchemaBinding& To, FLoadIdBinder& LoadIds, SubsetByteArray& Out)
+template<class LoadIdBinder>
+static void CloneBindingWithReplacedStructIds(const FSchemaId* FromIds, const FSchemaBinding& To, LoadIdBinder& LoadIds, SubsetByteArray& Out)
 {
 	check(IsAligned(Out.Num(), alignof(FSchemaBinding)));
 
@@ -333,7 +346,8 @@ static void CloneBindingWithReplacedStructIds(const FSchemaId* FromIds, const FS
 	}
 }
 
-[[nodiscard]] static FLoadStructPlan MakeSchemaLoadPlan(const FStructSchema& From, const FSchemaBinding& To, TConstArrayView<FMemberId> ToMemberIds, TConstArrayView<FStructSchemaId> ToStructIds, FLoadIdBinder& LoadIds, SubsetByteArray& OutSubsetSchemas)
+template<class LoadIdBinder>
+[[nodiscard]] static FLoadStructPlan MakeSchemaLoadPlan(const FStructSchema& From, const FSchemaBinding& To, TConstArrayView<FMemberId> ToMemberIds, TConstArrayView<FStructSchemaId> ToStructIds, LoadIdBinder& LoadIds, SubsetByteArray& OutSubsetSchemas)
 {
 	uint16 NumEnums = CountEnums(From);
 	if (From.NumMembers < To.NumMembers || NumEnums || HasDifferentSupers(From, To, ToStructIds))
@@ -363,22 +377,29 @@ static void CloneBindingWithReplacedStructIds(const FSchemaId* FromIds, const FS
 	return NullOpt;
 }
 
-[[nodiscard]] static FLoadStructPlan MakeLoadPlan(const FStructSchema& From, const FSchemaBinding& To, TConstArrayView<FMemberId> ToMemberIds, TConstArrayView<FStructSchemaId> ToStructIds, FLoadIdBinder& LoadIds, SubsetByteArray& OutSubsetSchemas)
+template<class LoadIdBinder>
+[[nodiscard]] static FLoadStructPlan MakeLoadPlan(const FStructSchema& From, const FSchemaBinding& To, TConstArrayView<FMemberId> ToMemberIds, TConstArrayView<FStructSchemaId> ToStructIds, LoadIdBinder& LoadIds, SubsetByteArray& OutSubsetSchemas)
 {
 	TOptional<FLoadStructMemcpy> Memcpy = TryMakeMemcpyPlan(From, To, ToMemberIds);
 	return Memcpy ? FLoadStructPlan(Memcpy.GetValue()) : MakeSchemaLoadPlan(From, To, ToMemberIds, ToStructIds, LoadIds, OutSubsetSchemas);
 }
 
-struct FMissingBinding : ICustomBinding
+struct FLoadBindingBase : ICustomBinding
+{
+	virtual void SaveCustom(FMemberBuilder& Dst, const void* Src, const void* Default, const FSaveContext& Ctx) override final { unimplemented(); }
+	virtual bool DiffCustom(const void*, const void*) const override final { unimplemented(); return false; }
+};
+struct FTypeErasedLoadBinding final : FLoadBindingBase
+{
+	virtual void LoadCustom(void* Dst, FStructLoadView Src, ECustomLoadMethod) const override { checkf(false, TEXT("Can't load type-erased/lowered struct binding with load id %d"), Src.Schema.LoadId.Idx); }
+};
+static FTypeErasedLoadBinding GLoadTypeErasedError;
+struct FNoopLoadBinding final : FLoadBindingBase
 {
-	FMissingBinding(const TCHAR* InType) : Type(InType) {}
-	const TCHAR* Type;
-	virtual void SaveCustom(FMemberBuilder& Dst, const void* Src, const void* Default, const FSaveContext& Ctx) override { checkf(false, TEXT("Can't save %s"), Type); }
-	virtual void LoadCustom(void* Dst, FStructLoadView Src, ECustomLoadMethod Method) const override { checkf(false, TEXT("Can't load %s"), Type); }
-	virtual bool DiffCustom(const void* StructA, const void* StructB) const override { checkf(false, TEXT("Can't diff %s"), Type); return false; }
+	virtual void LoadCustom(void* Dst, FStructLoadView Src, ECustomLoadMethod Method) const override {}
 };
-static FMissingBinding GMemberlessBinding(TEXT("memberless struct binding"));
-static FMissingBinding GUnboundBinding{TEXT("type-erased/lowered struct binding")};
+static FNoopLoadBinding GLoadNoop;
+
 
 template<class T>
 static TOptional<TConstArrayView<T>> GetContiguousSubset(TConstArrayView<T> View, const TBitArray<>& Keep)
@@ -424,7 +445,7 @@ struct FLoadPlanner
 	SubsetByteArray									SubsetSchemaData;
 	TSet<FStructSchemaId>							UnboundSaveIds;
 	
-	FLoadBatchPtr CreatePlans()
+	FLoadBatchPtr CreatePlans(ESchemaFormat Format)
 	{
 		check(NumStructSchemas(ReadId) == RuntimeIds.Num());
 
@@ -432,6 +453,22 @@ struct FLoadPlanner
 		const uint32 NumPlans = RuntimeIds.Num();
 		Plans.SetNumUninitialized(NumPlans);
 		SubsetSchemaSizes.SetNumUninitialized(NumPlans);
+
+		if (Format == ESchemaFormat::InMemoryNames)
+		{
+			FLoadIdBinderDummy LoadIds;
+			for (uint32 Idx = 0; Idx < NumPlans; ++Idx)
+			{
+				FLoadIdMapping Mapping;
+				Mapping.ReadId = { Idx };
+				Mapping.LoadId = { Idx };
+				Mapping.BindId = RuntimeIds[Idx];
+				CreatePlan(Mapping, LoadIds);
+			}
+
+			return CreateBatch(LoadIds);
+		}
+	
 		FLoadIdBinder LoadIds(RuntimeIds);
 		for (uint32 Idx = 0; Idx < NumPlans; ++Idx)
 		{
@@ -472,7 +509,8 @@ struct FLoadPlanner
 	}
 private:
 
-	FLoadBatchPtr CreateBatch(const FLoadIdBinder& LoadIds)
+	template<typename LoadIdBinder>
+	FLoadBatchPtr CreateBatch(const LoadIdBinder& LoadIds)
 	{
 		const uint32 NumPlans = Plans.Num();
 		const uint32 NumMappings = LoadIds.NumMappings();
@@ -529,7 +567,8 @@ private:
 		return FLoadBatchPtr(Out);
 	}
 
-	void CreatePlan(FLoadIdMapping Mapping, FLoadIdBinder& LoadIds)
+	template<class LoadIdBinder>
+	void CreatePlan(FLoadIdMapping Mapping, LoadIdBinder& LoadIds)
 	{
 		int32 SubsetSchemaOffset = SubsetSchemaData.Num();	
 		Plans[Mapping.LoadId.Idx] = CreatePlanInner(Mapping, LoadIds);
@@ -647,12 +686,20 @@ private:
 		return TmpSubset;
 	}
 
-	FLoadStructPlan CreatePlanInner(FLoadIdMapping Mapping, FLoadIdBinder& LoadIds)
+	template<class LoadIdBinder>
+	FLoadStructPlan CreatePlanInner(FLoadIdMapping Mapping, LoadIdBinder& LoadIds)
 	{
 		const FStructSchema& From = ResolveStructSchema(ReadId, Mapping.ReadId);
 		
 		TConstArrayView<FInnerStruct> LoweredInners;
-		if (const ICustomBinding* Custom = Customs.FindStruct(Mapping.BindId, /* out */ LoweredInners))
+		if constexpr (std::is_same_v<LoadIdBinder, FLoadIdBinderDummy>)
+		{
+			if (const ICustomBinding* Custom = Customs.FindStruct(Mapping.BindId))
+			{
+				return FLoadStructPlan(*Custom);
+			}
+		}
+		else if (const ICustomBinding* Custom = Customs.FindStruct(Mapping.BindId, /* out */ LoweredInners))
 		{
 			TArray<FInnerStruct, TInlineAllocator<8>> TmpSubset;
 			TConstArrayView<FInnerStruct> LoweredMembers = GetLoweredMembers(LoweredInners, From.GetMemberNames(), /* out */ TmpSubset);
@@ -676,20 +723,16 @@ private:
 			
 			// Type-erased structs
 			UnboundSaveIds.Add(Mapping.ReadId);
-			return FLoadStructPlan(GUnboundBinding);
+			return FLoadStructPlan(GLoadTypeErasedError);
 		}
 
-//		checkf(!From.Type.Scope && From.Type.Name.NumParameters == 2, TEXT("Unexpected bind id typename '%s' saved as '%s'. Only range-bound template parameters are memberless. "
-//				"They're always anonymous two-parameter types and uninstantiable as structs, bound via MakeAnonymousParametricType(). "),
-//				*Declarations.GetDebug().Print(From.Type),
-//				*Declarations.GetDebug().Print(Mapping.SaveId));
-		return FLoadStructPlan(GMemberlessBinding);
+		return FLoadStructPlan(GLoadNoop);
 	}
 };
 
-FLoadBatchPtr CreateLoadPlans(FReadBatchId ReadId, const FDeclarations& Declarations, const FCustomBindings& Customs, const FSchemaBindings& Schemas, TConstArrayView<FStructSchemaId> RuntimeIds)
+FLoadBatchPtr CreateLoadPlans(FReadBatchId ReadId, const FDeclarations& Declarations, const FCustomBindings& Customs, const FSchemaBindings& Schemas, TConstArrayView<FStructSchemaId> RuntimeIds, ESchemaFormat Format)
 {
-	return FLoadPlanner{ReadId, Declarations, Customs, Schemas, RuntimeIds}.CreatePlans();
+	return FLoadPlanner{ReadId, Declarations, Customs, Schemas, RuntimeIds}.CreatePlans(Format);
 }
 
 
@@ -925,11 +968,15 @@ private:
 		// Make code changes in FMemberReader::SkipMissingSparseMembers() too
 		while (bSparse && MemberIdx < Types.Num() && BitIt.GrabNext(ByteIt))
 		{
-			FMemberBindType Type = Types[MemberIdx];
-			FMemberBindType InnermostType = Type.IsRange() ? GrabInnerRanges(Type.AsRange()).InnerTypes.Last() : Type;
-			if (InnermostType.IsStruct())
+			FMemberBindType Type = Types[MemberIdx];		
+			EMemberKind Kind = Type.GetKind();
+			if (Kind == EMemberKind::Struct)
 			{
-				(void)GrabInnerStruct(InnermostType.AsStruct());
+				(void)GrabInnerStruct(Type.AsStruct());
+			}
+			else if (Kind == EMemberKind::Range)
+			{
+				(void)GrabInnerRanges(Type.AsRange());
 			}
 			++MemberIdx;
 		}
@@ -1000,6 +1047,7 @@ private:
 void LoadStruct(void* Dst, FByteReader Src, FStructSchemaId LoadId, const FLoadBatch& Batch)
 {
 	FLoadStructPlan Plan = Batch[LoadId];
+	
 	if (Plan.IsSchema())
 	{
 		if (Plan.IsSparseSchema())
diff --git a/Engine/Plugins/Experimental/PlainProps/Source/Private/PlainPropsWrite.cpp b/Engine/Plugins/Experimental/PlainProps/Source/Private/PlainPropsWrite.cpp
--- a/Engine/Plugins/Experimental/PlainProps/Source/Private/PlainPropsWrite.cpp
+++ b/Engine/Plugins/Experimental/PlainProps/Source/Private/PlainPropsWrite.cpp
@@ -266,6 +266,8 @@ FWriteIds::FWriteIds(const FIdIndexerBase& Ids, const IStructBindIds& InBindIds,
 			KeptParameter = Remap(KeptParameter);
 		}
 	}
+
+	check(HasStableNames() == (Format == ESchemaFormat::StableNames));
 }
 
 FStructSchemaId
diff --git a/Engine/Plugins/Experimental/PlainProps/Source/Public/PlainPropsBind.h b/Engine/Plugins/Experimental/PlainProps/Source/Public/PlainPropsBind.h
--- a/Engine/Plugins/Experimental/PlainProps/Source/Public/PlainPropsBind.h
+++ b/Engine/Plugins/Experimental/PlainProps/Source/Public/PlainPropsBind.h
@@ -1103,8 +1103,8 @@ void BindNativeStruct(FSchemaBindings& Out, FStructSchemaId BindId, FStructSchem
 
 //////////////////////////////////////////////////////////////////////////
 
-// Save -> load struct ids for ESchemaFormat::InMemoryNames
-[[nodiscard]] PLAINPROPS_API TArray<FStructSchemaId> IndexInMemoryNames(const FSchemaBatch& Schemas,  FIdIndexerBase& Indexer);
+// Save -> load struct ids for ESchemaFormat::InMemoryNames, alternative to side-channel with ExtractRuntimeIds()
+[[nodiscard]] PLAINPROPS_API TArray<FStructSchemaId> IndexRuntimeIds(const FSchemaBatch& Schemas,  FIdIndexerBase& Indexer);
 
 // Save -> load ids for ESchemaFormat::StableNames
 struct FIdBinding
@@ -1137,8 +1137,8 @@ struct FIdBinding
 
 struct FIdTranslatorBase
 {
-	static uint32 CalculateTranslationSize(int32 NumSavedNames, const FSchemaBatch& Batch);
-	static FIdBinding TranslateIds(FMutableMemoryView To, FIdIndexerBase& Indexer, TConstArrayView<FNameId> TranslatedNames, const FSchemaBatch& From);
+	PLAINPROPS_API static uint32 CalculateTranslationSize(int32 NumSavedNames, const FSchemaBatch& Batch);
+	PLAINPROPS_API static FIdBinding TranslateIds(FMutableMemoryView To, FIdIndexerBase& Indexer, TConstArrayView<FNameId> TranslatedNames, const FSchemaBatch& From);
 };
 
 // Maps saved ids -> runtime load ids for ESchemaFormat::StableNames
@@ -1151,7 +1151,7 @@ struct FIdTranslator : FIdTranslatorBase
 		
 		// Translate names
 		TArrayView<FNameId> NewNames(reinterpret_cast<FNameId*>(&Allocator[0]), SavedNames.Num());
-		FNameId* NameIt = &NewNames[0];
+		FNameId* NameIt = NewNames.GetData();
 		for (const NameType& SavedName : SavedNames)
 		{
 			(*NameIt++) = Indexer.MakeName(SavedName);
@@ -1165,7 +1165,7 @@ struct FIdTranslator : FIdTranslatorBase
 	TArray<uint8, TInlineAllocator<1024>>	Allocator;
 };
 
-FSchemaBatch*			CreateTranslatedSchemas(const FSchemaBatch& Schemas, FIdBinding NewIds);
-void					DestroyTranslatedSchemas(const FSchemaBatch* Schemas);
+PLAINPROPS_API FSchemaBatch*	CreateTranslatedSchemas(const FSchemaBatch& Schemas, FIdBinding NewIds);
+PLAINPROPS_API void				DestroyTranslatedSchemas(const FSchemaBatch* Schemas);
 
 } // namespace PlainProps
diff --git a/Engine/Plugins/Experimental/PlainProps/Source/Public/PlainPropsBuildSchema.h b/Engine/Plugins/Experimental/PlainProps/Source/Public/PlainPropsBuildSchema.h
--- a/Engine/Plugins/Experimental/PlainProps/Source/Public/PlainPropsBuildSchema.h
+++ b/Engine/Plugins/Experimental/PlainProps/Source/Public/PlainPropsBuildSchema.h
@@ -6,6 +6,7 @@
 #include "Containers/Array.h"
 #include "Containers/ArrayView.h"
 #include "Containers/PagedArray.h"
+#include "Containers/Set.h"
 
 namespace PlainProps
 {
@@ -46,8 +47,8 @@ public:
 	using FStructDeclarations = TConstArrayView<TUniquePtr<FStructDeclaration>>;
 	using FEnumDeclarations = TConstArrayView<TUniquePtr<FEnumDeclaration>>;
 
-	PLAINPROPS_API FSchemasBuilder(const FDeclarations& Declarations, const IStructBindIds& BindIds, FScratchAllocator& Scratch);
-	PLAINPROPS_API FSchemasBuilder(FStructDeclarations InStructs, FEnumDeclarations InEnums, const IStructBindIds& BindIds, const FDebugIds& InDebug, FScratchAllocator& Scratch);
+	PLAINPROPS_API FSchemasBuilder(const FDeclarations& Declarations, const IStructBindIds& BindIds, FScratchAllocator& Scratch, ESchemaFormat InFormat);
+	PLAINPROPS_API FSchemasBuilder(FStructDeclarations InStructs, FEnumDeclarations InEnums, const IStructBindIds& BindIds, const FDebugIds& InDebug, FScratchAllocator& Scratch, ESchemaFormat InFormat);
 	PLAINPROPS_API ~FSchemasBuilder();
 
 	PLAINPROPS_API FEnumSchemaBuilder&			NoteEnum(FEnumSchemaId Id);
@@ -61,9 +62,10 @@ public:
 private:
 	FStructDeclarations							DeclaredStructs;
 	FEnumDeclarations							DeclaredEnums;
-	TArray<int32>								StructIndices;
-	TArray<int32>								EnumIndices;
+	TSet<FStructSchemaId>						StructIndices;
+	TSet<FEnumSchemaId>							EnumIndices;
 	const IStructBindIds&						BindIds;
+	ESchemaFormat								Format;
 	TPagedArray<FStructSchemaBuilder, 4096>		Structs;		// TPagedArray for stable references
 	TPagedArray<FEnumSchemaBuilder, 4096>		Enums;			// TPagedArray for stable references
 	FScratchAllocator&							Scratch;
@@ -73,4 +75,7 @@ private:
 	void										NoteInheritanceChains();
 };
 
+// Extract runtime ids to side-channel for loading to avoid reindexing with IndexRuntimeIds()
+[[nodiscard]] PLAINPROPS_API TArray<FStructSchemaId> ExtractRuntimeIds(const FBuiltSchemas& In);
+
 } // namespace PlainProps
\ No newline at end of file
diff --git a/Engine/Plugins/Experimental/PlainProps/Source/Public/PlainPropsIndex.h b/Engine/Plugins/Experimental/PlainProps/Source/Public/PlainPropsIndex.h
--- a/Engine/Plugins/Experimental/PlainProps/Source/Public/PlainPropsIndex.h
+++ b/Engine/Plugins/Experimental/PlainProps/Source/Public/PlainPropsIndex.h
@@ -62,8 +62,8 @@ public:
 	
 	virtual FNestedScope				Resolve(FNestedScopeId Id) const override final		{ return NestedScopes.Resolve(Id); }
 	virtual FParametricTypeView			Resolve(FParametricTypeId Id) const override final	{ return ParametricTypes.Resolve(Id); }
-	virtual FTypeId						Resolve(FEnumSchemaId Id) const override final		{ return Enums.Get(FSetElementId::FromInteger(Id.Idx)); }
-	virtual FTypeId						Resolve(FStructSchemaId Id) const override final	{ return Structs.Get(FSetElementId::FromInteger(Id.Idx)); }
+	PLAINPROPS_API virtual FTypeId		Resolve(FEnumSchemaId Id) const override final;
+	PLAINPROPS_API virtual FTypeId		Resolve(FStructSchemaId Id) const override final;
 	
 	const FNestedScopeIndexer&			GetNestedScopes() const								{ return NestedScopes; }
 	const FParametricTypeIndexer&		GetParametricTypes() const							{ return ParametricTypes; }
diff --git a/Engine/Plugins/Experimental/PlainProps/Source/Public/PlainPropsLoad.h b/Engine/Plugins/Experimental/PlainProps/Source/Public/PlainPropsLoad.h
--- a/Engine/Plugins/Experimental/PlainProps/Source/Public/PlainPropsLoad.h
+++ b/Engine/Plugins/Experimental/PlainProps/Source/Public/PlainPropsLoad.h
@@ -21,14 +21,8 @@ struct FLoadBatchDeleter
 };
 using FLoadBatchPtr = TUniquePtr<FLoadBatch, FLoadBatchDeleter>;
 
-template<class Runtime>
-[[nodiscard]] FLoadBatchPtr					CreateLoadPlans(FReadBatchId ReadId, TConstArrayView<FStructSchemaId> RuntimeIds)
-{
-	return CreateLoadPlans(ReadId, Runtime::GetTypes(), Runtime::GetCustoms(), Runtime::GetSchemas(), RuntimeIds);
-}
-
 
-[[nodiscard]] PLAINPROPS_API FLoadBatchPtr	CreateLoadPlans(FReadBatchId ReadId, const FDeclarations& Declarations, const FCustomBindings& Customs, const FSchemaBindings& Schemas, TConstArrayView<FStructSchemaId> RuntimeIds);
+[[nodiscard]] PLAINPROPS_API FLoadBatchPtr	CreateLoadPlans(FReadBatchId ReadId, const FDeclarations& Declarations, const FCustomBindings& Customs, const FSchemaBindings& Schemas, TConstArrayView<FStructSchemaId> RuntimeIds, ESchemaFormat Format);
 PLAINPROPS_API void							LoadStruct(void* Dst, FByteReader Src, FStructSchemaId Id, const FLoadBatch& Batch);
 PLAINPROPS_API void							ConstructAndLoadStruct(void* Dst, FByteReader Src, FStructSchemaId LoadId, const FLoadBatch& Batch);
 
diff --git a/Engine/Plugins/Experimental/PlainProps/Source/Public/PlainPropsTypes.h b/Engine/Plugins/Experimental/PlainProps/Source/Public/PlainPropsTypes.h
--- a/Engine/Plugins/Experimental/PlainProps/Source/Public/PlainPropsTypes.h
+++ b/Engine/Plugins/Experimental/PlainProps/Source/Public/PlainPropsTypes.h
@@ -410,6 +410,10 @@ const T* AlignPtr(const void* Ptr)
 
 //////////////////////////////////////////////////////////////////////////
 
+enum class ESchemaFormat { StableNames, InMemoryNames };
+
+//////////////////////////////////////////////////////////////////////////
+
 class FDebugIds
 {
 public:
diff --git a/Engine/Plugins/Experimental/PlainProps/Source/Public/PlainPropsWrite.h b/Engine/Plugins/Experimental/PlainProps/Source/Public/PlainPropsWrite.h
--- a/Engine/Plugins/Experimental/PlainProps/Source/Public/PlainPropsWrite.h
+++ b/Engine/Plugins/Experimental/PlainProps/Source/Public/PlainPropsWrite.h
@@ -18,8 +18,6 @@ class FParametricTypeIndexer;
 struct FWriteIds;
 struct IStructBindIds;
 
-enum class ESchemaFormat { StableNames, InMemoryNames };
-
 class FWriter
 {
 public:
diff --git a/Engine/Plugins/Experimental/PlainProps/Source/Tests/PlainPropsTests.cpp b/Engine/Plugins/Experimental/PlainProps/Source/Tests/PlainPropsTests.cpp
--- a/Engine/Plugins/Experimental/PlainProps/Source/Tests/PlainPropsTests.cpp
+++ b/Engine/Plugins/Experimental/PlainProps/Source/Tests/PlainPropsTests.cpp
@@ -211,7 +211,7 @@ void FTestBatchBuilder::AddObject(FStructSchemaId Schema, FMemberBuilder&& Membe
 TArray64<uint8> FTestBatchBuilder::Write()
 {
 	// Build partial schemas
-	FSchemasBuilder SchemaBuilders(Declarations, *this, Scratch);
+	FSchemasBuilder SchemaBuilders(Declarations, *this, Scratch, ESchemaFormat::StableNames);
 	for (const TPair<FStructSchemaId, FBuiltStruct*>& Object : Objects)
 	{
 		SchemaBuilders.NoteStructAndMembers(Object.Key, *Object.Value);
diff --git a/Engine/Plugins/Experimental/PlainProps/Source/Tests/PlainPropsUeCoreTests.cpp b/Engine/Plugins/Experimental/PlainProps/Source/Tests/PlainPropsUeCoreTests.cpp
--- a/Engine/Plugins/Experimental/PlainProps/Source/Tests/PlainPropsUeCoreTests.cpp
+++ b/Engine/Plugins/Experimental/PlainProps/Source/Tests/PlainPropsUeCoreTests.cpp
@@ -148,7 +148,7 @@ public:
 	template<class T>
 	bool						SaveDelta(const T& Object, const T& Default);
 
-	TArray64<uint8>				Write() const;
+	TArray64<uint8>				Write(TArray<FStructSchemaId>* OutMemoryIds) const;
 
 private:
 	using IdBuiltStructPair = TPair<FStructSchemaId, const FBuiltStruct*>;
@@ -199,19 +199,25 @@ TConstArrayView<T> GrabNumAndArray(/* in-out */ FByteReader& It)
 
 inline constexpr uint32 Magics[] = { 0xFEEDF00D, 0xABCD1234, 0xDADADAAA, 0x99887766, 0xF0F1F2F3 , 0x00112233};
 
-TArray64<uint8> FBatchSaver::Write() const
+TArray64<uint8> FBatchSaver::Write(TArray<FStructSchemaId>* OutMemoryIds) const
 {
+	ESchemaFormat Format = OutMemoryIds ? ESchemaFormat::InMemoryNames : ESchemaFormat::StableNames;
+
 	// Build partial schemas
 	const FStructBindIds BindIds(Customs, GSchemas);
-	FSchemasBuilder SchemaBuilders(GTypes, BindIds, Scratch);
+	FSchemasBuilder SchemaBuilders(GTypes, BindIds, Scratch, Format);
 	for (const IdBuiltStructPair& Object : SavedObjects)
 	{
 		SchemaBuilders.NoteStructAndMembers(Object.Key, *Object.Value);
 	}
-	FBuiltSchemas Schemas = SchemaBuilders.Build(); 
+	FBuiltSchemas Schemas = SchemaBuilders.Build();
+	if (OutMemoryIds)
+	{
+		*OutMemoryIds = ExtractRuntimeIds(Schemas);
+	}
 
 	// Filter out declared but unused names and ids
-	FWriter Writer(GNames, BindIds, Schemas, ESchemaFormat::StableNames);
+	FWriter Writer(GNames, BindIds, Schemas, Format);
 	TArray<FName> UsedNames;
 	UsedNames.Reserve(Writer.GetUsedNames().Num());
 	for (FNameId Name : Writer.GetUsedNames())
@@ -256,22 +262,40 @@ TArray64<uint8> FBatchSaver::Write() const
 
 ////////////////////////////////////////////////////////////////////////////////////////////////
 
-class FSavedBatchNames : public FDebugIds
+class FSchemaBatchIds : public FDebugIds
 {
 public:
-	FSavedBatchNames(TConstArrayView<FName> InNames, const FSchemaBatch& InSchemas) : Names(InNames), Schemas(InSchemas) {}
+	FSchemaBatchIds(const FSchemaBatch& InSchemas) : Schemas(InSchemas) {}
 
-	virtual FNestedScope				Resolve(FNestedScopeId Id) const override final		{ return ResolveNestedScope(Schemas, Id); }
-	virtual FParametricTypeView			Resolve(FParametricTypeId Id) const override final	{ return ResolveParametricType(Schemas, Id); }
 	virtual FTypeId						Resolve(FEnumSchemaId Id) const override final		{ return ResolveEnumSchema(Schemas, Id).Type; }
 	virtual FTypeId						Resolve(FStructSchemaId Id) const override final	{ return ResolveStructSchema(Schemas, Id).Type; }
 
-	using FDebugIds::AppendString;
-	virtual void AppendString(FUtf8StringBuilderBase& Out, FNameId Name) const override final { Names[Name.Idx].AppendString(Out); }
+protected:
+	const FSchemaBatch& Schemas;
+};
+
+class FMemoryBatchIds final : public FSchemaBatchIds
+{
+public:
+	using FSchemaBatchIds::FSchemaBatchIds;
+	using FSchemaBatchIds::Resolve;
+
+	virtual FNestedScope		Resolve(FNestedScopeId Id) const override { return GNames.Resolve(Id); }
+	virtual FParametricTypeView	Resolve(FParametricTypeId Id) const override { return GNames.Resolve(Id); }
+	virtual void				AppendString(FUtf8StringBuilderBase& Out, FNameId Name) const override { GNames.ResolveName(Name).AppendString(Out); }
+};
+
+class FStableBatchIds final : public FSchemaBatchIds
+{
+public:
+	using FSchemaBatchIds::FSchemaBatchIds;
+	using FSchemaBatchIds::Resolve;
+
+	virtual FNestedScope		Resolve(FNestedScopeId Id) const override { return ResolveNestedScope(Schemas, Id); }
+	virtual FParametricTypeView	Resolve(FParametricTypeId Id) const override { return ResolveParametricType(Schemas, Id); }
+	virtual void				AppendString(FUtf8StringBuilderBase& Out, FNameId Name) const override { Names[Name.Idx].AppendString(Out); }
 
-private:
 	TConstArrayView<FName> Names;
-	const FSchemaBatch& Schemas;
 };
 
 class FLoadBatchNames : public FDebugIds
@@ -295,14 +319,13 @@ private:
 	const FSchemaBatch& Schemas;
 };
 
-void PrintSavedSchemas(TConstArrayView<FName> Names, const FSchemaBatch& SavedSchemas)
+void PrintSavedSchemas(const FSchemaBatchIds& BatchIds, const FSchemaBatch& SavedSchemas)
 {
 	FReadBatchId Batch = MountReadSchemas(&SavedSchemas);
 	TUtf8StringBuilder<4096> YamlString;
 	{
 		FYamlBuilderPtr YamlBuilder = MakeYamlBuilder(YamlString);
-		FSavedBatchNames BatchNames(Names, SavedSchemas);
-		FBatchPrinter Printer(*YamlBuilder, BatchNames);
+		FBatchPrinter Printer(*YamlBuilder, BatchIds);
 		Printer.PrintSchemas(Batch);
 	}
 	FUtf8StringView YamlView = YamlString.ToView();
@@ -311,6 +334,21 @@ void PrintSavedSchemas(TConstArrayView<FName> Names, const FSchemaBatch& SavedSc
 	UnmountReadSchemas(Batch);
 }
 
+void PrintSavedSchemas(TConstArrayView<FName> SavedNames, const FSchemaBatch& SavedSchemas)
+{
+	if (SavedNames.IsEmpty())
+	{
+		FMemoryBatchIds BatchIds(SavedSchemas);
+		PrintSavedSchemas(BatchIds, SavedSchemas);
+	}
+	else
+	{
+		FStableBatchIds BatchIds(SavedSchemas);
+		BatchIds.Names = SavedNames;
+		PrintSavedSchemas(BatchIds, SavedSchemas);
+	}
+}
+
 void PrintLoadStruct(const FIdBinding& Binding, FStructView StructView, int32 LoadIdx)
 {
 	TUtf8StringBuilder<4096> YamlString;
@@ -330,14 +368,16 @@ void PrintLoadStruct(const FIdBinding& Binding, FStructView StructView, int32 Lo
 class FBatchLoader
 {
 public:
-	FBatchLoader(FMemoryView Data, const FCustomBindings& Underlay)
+	FBatchLoader(FMemoryView Data, const FCustomBindings& Underlay, TConstArrayView<FStructSchemaId> InRuntimeIds)
 	: Customs(Underlay)
 	{
+		ESchemaFormat Format = InRuntimeIds.IsEmpty() ? ESchemaFormat::StableNames : ESchemaFormat::InMemoryNames;
+
 		// Read ids
 		FByteReader It(Data);
 		CHECK(It.Grab<uint32>() == Magics[0]);
 		Ids = GrabNumAndArray<FName>(It);
-		CHECK(Ids.Num() != 0);
+		CHECK(Ids.IsEmpty() != InRuntimeIds.IsEmpty());
 		
 		// Read schemas
 		CHECK(It.Grab<uint32>() == Magics[1]);
@@ -348,15 +388,28 @@ public:
 
 		// Mount, print and unmount saved schemas
 		PrintSavedSchemas(Ids, *SavedSchemas);
-		
-		// Bind saved ids to runtime ids, make new schemas with new ids and mount them
-		FIdTranslator RuntimeIds(GNames, Ids, *SavedSchemas);
-		FSchemaBatch* LoadSchemas = CreateTranslatedSchemas(*SavedSchemas, RuntimeIds.Translation);
-		FReadBatchId Batch = MountReadSchemas(LoadSchemas);
 
-		// Store mappings required for structs printing of untranslated scopes and parametric types
-		IdBinding = RuntimeIds.Translation;
+		FReadBatchId Batch;
+		TOptional<FIdTranslator> RuntimeIds;
+		if (InRuntimeIds.IsEmpty())
+		{
+			// Bind saved ids to runtime ids, make new schemas with new ids and mount them
+			RuntimeIds.Emplace(GNames, Ids, *SavedSchemas);
+			const FSchemaBatch* LoadSchemas = CreateTranslatedSchemas(*SavedSchemas, RuntimeIds->Translation);
+			Batch = MountReadSchemas(LoadSchemas);
+
+			// Store mappings required for structs printing of untranslated scopes and parametric types
+			IdBinding.Emplace(RuntimeIds->Translation);
+		}
+		else
+		{
+			// Mount saved schemas as is
+			Batch = MountReadSchemas(SavedSchemas);	
 
+			TArray<FStructSchemaId> lala = IndexRuntimeIds(*SavedSchemas, GNames);
+			check(EqualItems(lala, InRuntimeIds));
+		}
+		
 		// Read objects
 		while (uint64 NumBytes = It.GrabVarIntU())
 		{	
@@ -375,8 +428,8 @@ public:
 		CHECK(It.Grab<uint32>() == Magics[5]);
 
 		// Finally create load plans
-		TConstArrayView<FStructSchemaId> LoadStructIds = RuntimeIds.Translation.GetStructIds(SavedSchemas->NumStructSchemas);
-		Plans = CreateLoadPlans(Batch, GTypes, Customs, GSchemas, LoadStructIds);
+		TConstArrayView<FStructSchemaId> LoadStructIds = RuntimeIds ? RuntimeIds->Translation.GetStructIds(SavedSchemas->NumStructSchemas) : InRuntimeIds;
+		Plans = CreateLoadPlans(Batch, GTypes, Customs, GSchemas, LoadStructIds, Format);
 	}
 
 	~FBatchLoader()
@@ -384,7 +437,10 @@ public:
 		CHECK(LoadIdx == Objects.Num()); // Test should load all saved objects
 		Plans.Reset();
 		const FSchemaBatch* LoadSchemas = UnmountReadSchemas(Objects[0].Schema.Batch);
-		DestroyTranslatedSchemas(LoadSchemas);
+		if (IdBinding)
+		{
+			DestroyTranslatedSchemas(LoadSchemas);
+		}
 	}
 
 	template<class T>
@@ -407,7 +463,10 @@ public:
 	void LoadInto(T& Out)
 	{
 		FStructView In = Objects[LoadIdx];
-		PrintLoadStruct(IdBinding, In, LoadIdx);
+		if (IdBinding)
+		{
+			PrintLoadStruct(IdBinding.GetValue(), In, LoadIdx);
+		}
 		LoadStruct(&Out, In.Values, In.Schema.Id, *Plans);
 		++LoadIdx;
 	}
@@ -415,7 +474,7 @@ public:
 private:
 	TConstArrayView<FName>		Ids;
 	FNameBinding				Names;
-	FIdBinding					IdBinding;
+	TOptional<FIdBinding>		IdBinding;
 	FCustomBindingsOverlay		Customs;
 	FLoadBatchPtr				Plans;
 	TArray<FStructView>			Objects;
@@ -425,15 +484,19 @@ private:
 
 static void Run(void (*Save)(FBatchSaver&), void (*Load)(FBatchLoader&), const FCustomBindings& Customs = GCustoms)
 {
-	TArray64<uint8> Data;
+	for (ESchemaFormat Format : {ESchemaFormat::StableNames, ESchemaFormat::InMemoryNames})
 	{
-		FBatchSaver Batch(Customs);
-		Save(Batch);
-		Data = Batch.Write();
-	}
+		TArray64<uint8> Data;
+		TArray<FStructSchemaId>	RuntimeIds;
+		{
+			FBatchSaver Batch(Customs);
+			Save(Batch);
+			Data = Batch.Write(Format == ESchemaFormat::InMemoryNames ? &RuntimeIds : nullptr);
+		}
 
-	FBatchLoader Batch(MakeMemoryView(Data), Customs);
-	Load(Batch);
+		FBatchLoader Batch(MakeMemoryView(Data), Customs, MakeArrayView(RuntimeIds));
+		Load(Batch);
+	}
 }
 
 //////////////////////////////////////////////////////////////////////////
@@ -675,19 +738,16 @@ template <> struct TTypename<PlainProps::UE::Test::FIntAlias>
 {
 	inline static constexpr std::string_view DeclName = "FInt";
 	inline static constexpr std::string_view BindName = "IntAlias";
-	inline static constexpr std::string_view Namespace;
 };
 template <> struct TTypename<PlainProps::UE::Test::FSame1>
 {
 	inline static constexpr std::string_view DeclName = "Same";
 	inline static constexpr std::string_view BindName = "Same1";
-	inline static constexpr std::string_view Namespace;
 };
 template <> struct TTypename<PlainProps::UE::Test::FSame2>
 {
 	inline static constexpr std::string_view DeclName = "Same";
 	inline static constexpr std::string_view BindName = "Same2";
-	inline static constexpr std::string_view Namespace;
 };
 
 } namespace PlainProps::UE::Test {
diff --git a/Engine/Plugins/Experimental/PlainPropsUObject/Source/Private/PlainPropsUObjectBindings.cpp b/Engine/Plugins/Experimental/PlainPropsUObject/Source/Private/PlainPropsUObjectBindings.cpp
--- a/Engine/Plugins/Experimental/PlainPropsUObject/Source/Private/PlainPropsUObjectBindings.cpp
+++ b/Engine/Plugins/Experimental/PlainPropsUObject/Source/Private/PlainPropsUObjectBindings.cpp
@@ -383,7 +383,10 @@ inline void DeclareMembers(FMemberArray& Out, const UStruct* Struct)
 {
 	for (FProperty* It = Struct->PropertyLink; It && It->GetOwner<UStruct>() == Struct; It = It->PropertyLinkNext)
 	{
-		Out.Add(GUE.Names.NameMember(It->GetFName()));
+		if (ShouldBind(It))
+		{
+			Out.Add(GUE.Names.NameMember(It->GetFName()));
+		}
 	}
 }
 
@@ -454,7 +457,7 @@ void FFieldPathBinding::Load(FFieldPath& Dst, FStructLoadView Src, ECustomLoadMe
 {
 	FMemberLoader Members(Src);
 	Dst.Reset(); // ClearCachedField() more optimal
-	LoadStruct(&(Dst.*GFieldPathPath), Members.GrabStruct());
+	LoadNames(Dst.*GFieldPathPath, Members.GrabRange().AsStructs());
 	LoadStruct(&(Dst.*GFieldPathOwner), Members.GrabStruct());
 }
 
@@ -743,12 +746,13 @@ struct FInterfaceBinding final : ICustomBinding
 	
 	void Save(FMemberBuilder& Dst, const FScriptInterface& Src, const FSaveContext& Ctx) const
 	{
-		Dst.AddStruct(GUE.Members.Object, GUE.Structs.ObjectPtr, SaveStruct(Src.GetObject(), GUE.Structs.ObjectPtr, Ctx));
+		const TObjectPtr<UObject>& ObjectRef = const_cast<FScriptInterface&>(Src).GetObjectRef();
+		Dst.AddStruct(GUE.Members.Object, GUE.Structs.ObjectPtr, SaveStruct(&ObjectRef, GUE.Structs.ObjectPtr, Ctx));
 	}
 
 	void Load(FScriptInterface& Dst, FStructLoadView Src) const
 	{
-		LoadStruct(Dst.GetObjectRef(), FMemberLoader(Src).GrabStruct());
+		LoadSoleStruct(&Dst.GetObjectRef(), Src);
 		UObject* Object = Dst.GetObject();
 		Dst.SetInterface(Object ? Object->GetInterfaceAddress(InterfaceClass) : nullptr);
 	}
@@ -2974,7 +2978,7 @@ static void InitBatchedProperties()
 
 struct FMemoryPropertyBatch
 {
-	TArray<FText> Texts;
+	TArray<FText> Texts; // Tricky to serialize intrusively
 
 	static void Save(FMemberBuilder& Out, FName In, const FSaveContext&)
 	{
@@ -3133,15 +3137,21 @@ struct TCustomPropertyBindings
 	TCustomPropertyBinding<FLazyObjectPtr, BatchType>	LazyObjectPtr;
 };
 
-template<class PropertyBatchType>
-class TBatchSaver
+struct FMemoryBatch
+{
+	TArray64<uint8>				Data;
+	TArray<FStructSchemaId>		RuntimeIds; // To avoid reindexing schema FTypeId
+	FMemoryPropertyBatch		Properties; // Contains FTexts referenced from Data
+};
+
+static constexpr uint32 Magics[] = { 0xFEEDF00D, 0xABCD1234, 0xDADADAAA, 0x99887766, 0xF0F1F2F3 };
+
+class FBatchSaver
 {
 public:
-	TBatchSaver(int32 NumReserve)
-	: FlatBindings(SavedProperties, GUE.Customs)
-	, DeltaBindings(SavedProperties, GUE.DeltaCustoms)
-	, FlatCtx{GUE.Types, GUE.Schemas, FlatBindings.Overlay, Scratch, nullptr}
-	, DeltaCtx{GUE.Types, GUE.Schemas, DeltaBindings.Overlay, Scratch, &GUE.Defaults}
+	FBatchSaver(FCustomBindings& Flat, FCustomBindings& Delta, int32 NumReserve)
+	: FlatCtx{GUE.Types, GUE.Schemas, Flat, Scratch, nullptr}
+	, DeltaCtx{GUE.Types, GUE.Schemas, Delta, Scratch, &GUE.Defaults}
 	{
 		SavedObjects.Reserve(NumReserve);
 	}
@@ -3155,32 +3165,23 @@ public:
 		SavedObjects.Emplace(Id, bFlat ? SaveStruct(Object, Id, FlatCtx) : SaveStructDelta(Object, CDO, Id, DeltaCtx));
 	}
 
-	TArray64<uint8>	Write() const
+	TArray64<uint8>	Write(TArray<FStructSchemaId>& OutRuntimeIds) const
 	{
 		// Build partial schemas
 		const FStructBindIds BindIds(GUE.DeltaCustoms, GUE.Schemas);
-		FSchemasBuilder SchemaBuilders(GUE.Types, BindIds, Scratch);
+		FSchemasBuilder SchemaBuilders(GUE.Types, BindIds, Scratch, ESchemaFormat::InMemoryNames);
 		for (const IdBuiltStructPair& Object : SavedObjects)
 		{
 			SchemaBuilders.NoteStructAndMembers(Object.Key, *Object.Value);
 		}
-		FBuiltSchemas Schemas = SchemaBuilders.Build(); 
+		FBuiltSchemas Schemas = SchemaBuilders.Build();
 
-		// Filter out declared but unused names and ids
-		FWriter Writer(GUE.Names, BindIds, Schemas, ESchemaFormat::InMemoryNames);
-		TArray<FName> UsedNames;
-		UsedNames.Reserve(Writer.GetUsedNames().Num());
-		for (FNameId Name : Writer.GetUsedNames())
-		{
-			UsedNames.Add(GUE.Names.ResolveName(Name));
-		}
-	
-		// Write ids. Just copying in-memory FNames, a stable format might use SaveNameBatch().
-		TArray64<uint8> Out;
-		WriteInt(Out, Magics[0]);
-		WriteNumAndArray(Out, TArrayView<const FName, int32>(UsedNames));
+		// Save schema ids on the side
+		OutRuntimeIds = ExtractRuntimeIds(Schemas);
 
 		// Write schemas
+		FWriter Writer(GUE.Names, BindIds, Schemas, ESchemaFormat::InMemoryNames);
+		TArray64<uint8> Out;
 		WriteInt(Out, Magics[1]);
 		WriteAlignmentPadding<uint32>(Out);
 		TArray64<uint8> Tmp;
@@ -3208,14 +3209,9 @@ public:
 
 private:
 	using IdBuiltStructPair = TPair<FStructSchemaId, FBuiltStruct*>;
-	using FPropertyBindings = TCustomPropertyBindings<PropertyBatchType>;
-	static constexpr uint32 Magics[] = { 0xFEEDF00D, 0xABCD1234, 0xDADADAAA, 0x99887766, 0xF0F1F2F3 };
 
 	TArray<IdBuiltStructPair>	SavedObjects;
-	PropertyBatchType			SavedProperties;
 	mutable FScratchAllocator	Scratch;
-	FPropertyBindings			FlatBindings;
-	FPropertyBindings			DeltaBindings;	
 	FSaveContext				FlatCtx;
 	FSaveContext				DeltaCtx;
 
@@ -3225,7 +3221,64 @@ private:
 		WriteInt(Out, IntCastChecked<uint32>(Items.Num()));
 		WriteArray(Out, Items);
 	}
+};
+
+class FMemoryBatchLoader
+{
+public:
+	FMemoryBatchLoader(const FCustomBindings& Customs, FMemoryView Data, TConstArrayView<FStructSchemaId> RuntimeIds)
+	{
+		//// Read ids
+		
+		// Read and mount schemas
+		FByteReader It(Data);
+		check(It.Grab<uint32>() == Magics[1]);
+		It.SkipAlignmentPadding<uint32>();
+		uint32 SchemasSize = It.Grab<uint32>();
+		const FSchemaBatch* SavedSchemas = ValidateSchemas(It.GrabSlice(SchemasSize));
+		check(It.Grab<uint32>() == Magics[2]);
+		
+		FReadBatchId Batch = MountReadSchemas(SavedSchemas);
+
+		// Read objects
+		TSet<FStructSchemaId> SchemaIds;
+		SchemaIds.Reserve(NumStructSchemas(Batch));
+ 		while (uint64 NumBytes = It.GrabVarIntU())
+		{	
+			FByteReader ObjIt(It.GrabSlice(NumBytes));
+			check(ObjIt.Grab<uint32>() == Magics[3]);
+			FStructSchemaId Id = { ObjIt.Grab<uint32>() };
+			Objects.Add({ { Id, Batch }, ObjIt });
+			SchemaIds.Add(Id);
+		}
+		
+		check(It.Grab<uint32>() == Magics[4]);
+		check(!Objects.IsEmpty());
+
+		// Finally create load plans
+		Plans = CreateLoadPlans(Batch, GUE.Types, Customs, GUE.Schemas, RuntimeIds, ESchemaFormat::InMemoryNames);
+	}
+
+	~FMemoryBatchLoader()
+	{
+		check(LoadIdx == Objects.Num()); // Test should load all saved objects
+		Plans.Reset();
+		const FSchemaBatch* LoadSchemas = UnmountReadSchemas(Objects[0].Schema.Batch);
+		DestroyTranslatedSchemas(LoadSchemas);
+	}
+
+	void Load(UObject* Dst)
+	{
+		FStructView In = Objects[LoadIdx];
+		LoadStruct(Dst, In.Values, In.Schema.Id, *Plans);
+		++LoadIdx;
+	}
 
+private:
+	FLoadBatchPtr				Plans;
+	TArray<FStructView>			Objects;
+	int32						LoadIdx = 0;
+	
 	template<typename T>
 	static TConstArrayView<T> GrabNumAndArray(/* in-out */ FByteReader& It)
 	{
@@ -3234,19 +3287,57 @@ private:
 	}
 };
 
+struct FInstance
+{
+	UObject* Original;
+	UObject* PP;
+	UObject* TPS;
+};
 
 static void Roundtrip()
 {
-	TArray64<uint8> Data;
+	TArray<FInstance> Instances;
+	for (TObjectIterator<UObject> It; It; ++It)
 	{
-		TBatchSaver<FMemoryPropertyBatch> Batch(GUObjectArray.GetObjectArrayNum());
-		for (TObjectIterator<UObject> It; It; ++It)
+		UObject* Object = *It;
+		if (ShouldBind(Object->GetClass()))
 		{
-			Batch.Save(*It);
+			Instances.Add({*It});
 		}
-		Data = Batch.Write();
+	}
+	FName PP_N("PP", 1);
+	for (FInstance& Instance : Instances)
+	{
+		//todo: duplicate CDO somehow
+		UObject* O = Instance.Original;
+		PP_N.SetNumber(1 + &Instance - Instances.GetData());
+		Instance.PP = StaticAllocateObject(O->GetClass(), O->GetOuter(), PP_N, O->GetFlags());
+	}
+	FName TPS_N("TPS", 1);
+	for (FInstance& Instance : Instances)
+	{
+		UObject* O = Instance.Original;
+		TPS_N.SetNumber(1 + &Instance - Instances.GetData());
+		Instance.TPS = StaticAllocateObject(O->GetClass(), O->GetOuter(), TPS_N, O->GetFlags());
 	}
 
+	FMemoryBatch Doc;
+	TCustomPropertyBindings<FMemoryPropertyBatch> FlatBindings(Doc.Properties, GUE.Customs);
+	TCustomPropertyBindings<FMemoryPropertyBatch> DeltaBindings(Doc.Properties, GUE.DeltaCustoms);
+	{
+		FBatchSaver Batch(FlatBindings.Overlay, DeltaBindings.Overlay, GUObjectArray.GetObjectArrayNum());
+		for (const FInstance& Instance : Instances)
+		{
+			Batch.Save(Instance.Original);
+		}
+		Doc.Data = Batch.Write(/* out */ Doc.RuntimeIds);
+	}
+
+	FMemoryBatchLoader Batch(DeltaBindings.Overlay, MakeMemoryView(Doc.Data), Doc.RuntimeIds);
+	for (const FInstance& Instance : Instances)
+	{
+		Batch.Load(Instance.PP);
+	}
 }
 
 static void TestBindings()
