PlainProps: Generate and write deterministic nested scope and parametric type indices when using stable names.

#rnx

diff --git a/Engine/Plugins/Experimental/PlainProps/Source/Private/PlainPropsWrite.cpp b/Engine/Plugins/Experimental/PlainProps/Source/Private/PlainPropsWrite.cpp
--- a/Engine/Plugins/Experimental/PlainProps/Source/Private/PlainPropsWrite.cpp
+++ b/Engine/Plugins/Experimental/PlainProps/Source/Private/PlainPropsWrite.cpp
@@ -73,6 +73,8 @@ struct FUsedIds
 	TBitArray<> NestedScopes;
 	TBitArray<> ParametricTypes;
 	TArray<FNameId> KeptNames;
+	TArray<FNestedScopeId> KeptNestedScopes;
+	TArray<FParametricTypeId> KeptParametricTypes;
 
 
 	explicit FUsedIds(const FIdIndexerBase& InIds)
@@ -82,6 +84,8 @@ struct FUsedIds
 	, ParametricTypes(false, Ids.GetParametricTypes().Num())
 	{
 		KeptNames.Reserve(16 * (Ids.NumEnums() + Ids.NumStructs()));
+		KeptNestedScopes.Reserve(Ids.GetNestedScopes().Num());
+		KeptParametricTypes.Reserve(Ids.GetParametricTypes().Num());
 	}
 
 	template<class PartialSchemaType>
@@ -100,9 +104,10 @@ struct FUsedIds
 
 	void MarkUsed(FNameId Name)
 	{
-		if (!Names[Name.Idx])
+		FBitReference Used = Names[Name.Idx];
+		if (!Used)
 		{
-			Names[Name.Idx] = true;
+			Used = true;
 			KeptNames.Add(Name);
 		}
 	}
@@ -133,6 +138,7 @@ struct FUsedIds
 			if (!Used)
 			{
 				Used = true;
+				KeptNestedScopes.Add(Scope.AsNested());
 
 				FNestedScope Nested = Ids.Resolve(Scope.AsNested());
 				MarkUsed(Nested.Outer);
@@ -160,44 +166,59 @@ struct FUsedIds
 				{
 					MarkUsed(Parameter);
 				}
+				KeptParametricTypes.Add(Typename.AsParametric());
 			}
 		}
 	}
 };
 
 template<typename IdType>
-static uint32 MakeRemapping(TArray<TOptionalId<IdType>>& Out, const TBitArray<>& Used)
+static void MakeIdRemapping(TArray<TOptionalId<IdType>>& Out, const TArray<IdType>& Kept, int32 NumDeclared)
 {
-	uint32 NewIdx = 0;
-	Out.Reserve(Used.Num());
-	for (bool bUse : Used)
+	uint32 Idx = 0;
+	Out.SetNum(NumDeclared);
+	for (IdType Id : Kept)
 	{
-		Out.Add(bUse ? ToOptional(FromIdx<IdType>(NewIdx++)) : NoId);	
+		Out[Id.Idx] = FromIdx<IdType>(Idx++);
 	}
-	return NewIdx;
 }
 
-static uint32 MakeRemapping(TArray<FOptionalParametricTypeId>& Out, const TBitArray<>& Used, const FParametricTypeIndexer& Declared)
+static uint32 MakeIdRemapping(
+	TArray<FOptionalParametricTypeId>& Out,
+	const TArray<FParametricTypeId>& Kept,
+	const FParametricTypeIndexer& Declared)
 {
-	uint32 NewIdx = 0;
-	Out.Reserve(Used.Num());
-	for (bool bUse : Used)
+	uint32 SumParameters = 0;
+	uint32 Idx = 0;
+	Out.SetNum(Declared.Num());
+	for (FParametricTypeId Id : Kept)
 	{
-		uint8 NumParameters = static_cast<uint8>(Declared.At(Out.Num()).Parameters.NumParameters);
-		Out.Add(bUse ? ToOptional(FParametricTypeId(NumParameters, NewIdx++)) : NoId);	
+		uint8 NumParameters = static_cast<uint8>(Declared.At(Id.Idx).Parameters.NumParameters);
+		Out[Id.Idx] = FParametricTypeId(NumParameters, Idx++);
+		SumParameters += NumParameters;
 	}
-	return NewIdx;
+	return SumParameters;
+}
+
+
+inline static FNestedScope Resolve(FNestedScopeId Id, const FNestedScopeIndexer& Indexer)
+{
+	return Indexer.Resolve(Id);
+}
+
+inline static FParametricType Resolve(FParametricTypeId Id, const FParametricTypeIndexer& Indexer)
+{
+	return Indexer.At(Id.Idx);
 }
-template<typename T, typename C>
-static void CopyUsedIds(TArray<T>& Out, const TBitArray<>& Used, C&& Ids)
+
+template<typename ItemType, typename IdType, typename IndexerType>
+static void CopyUsed(TArray<ItemType>& Out, const TArray<IdType>& Kept, IndexerType&& Indexer)
 {
 	uint32 Idx = 0;
-	for (T Id : Ids)
+	Out.SetNumUninitialized(Kept.Num());
+	for (IdType Id : Kept)
 	{
-		if (Used[Idx++])
-		{
-			Out.Add(Id);
-		}
+		Out[Idx++] = Resolve(Id, Indexer);
 	}
 }
 
@@ -223,25 +244,19 @@ FWriteIds::FWriteIds(const FIdIndexerBase& Ids, const IStructBindIds& InBindIds,
 
 	if (Format == ESchemaFormat::StableNames)
 	{
-		// Generate new name indices in the deterministic order of traversing the built schemas
+		// Generate new name, nested scope and parametric type indices
+		// in the deterministic order of traversing the built schemas
 		FUsedIds Used(Ids);
 		Used.DetectUsage(Schemas.Structs);
 		Used.DetectUsage(Schemas.Enums);
 
+		// Remap ids and copy used names
+		MakeIdRemapping(/* out */ Names, Used.KeptNames, Ids.NumNames());
 		KeptNames = MoveTemp(Used.KeptNames);
-		Names.SetNum(Ids.NumNames());
-		for (uint32 Idx = 0, Num = IntCastChecked<uint32>(KeptNames.Num()); Idx < Num; ++Idx)
-		{
-			Names[KeptNames[Idx].Idx] = FNameId{ Idx };
-		}
 
-		// Generate new nested scope and parametric type mappings
-		KeptScopes.Reserve(MakeRemapping(/* out */ NestedScopes, Used.NestedScopes));
-		KeptParametrics.Reserve(MakeRemapping(/* out */ ParametricTypes, Used.ParametricTypes, Ids.GetParametricTypes()));
-
-		// Copy nested scopes and parametric types
-		CopyUsedIds(/* out */ KeptScopes, Used.NestedScopes, Ids.GetNestedScopes());
-		CopyUsedIds(/* out */ KeptParametrics, Used.ParametricTypes, Ids.GetParametricTypes().GetAllTypes());
+		// Remap ids and copy used nested scopes
+		MakeIdRemapping(/* out */ NestedScopes, Used.KeptNestedScopes, Ids.GetNestedScopes().Num());
+		CopyUsed(/* out */ KeptScopes, Used.KeptNestedScopes, Ids.GetNestedScopes());
 
 		// Remap copied nested scopes
 		for (FNestedScope& KeptScope : KeptScopes)
@@ -250,8 +265,13 @@ FWriteIds::FWriteIds(const FIdIndexerBase& Ids, const IStructBindIds& InBindIds,
 			KeptScope.Outer = Remap(KeptScope.Outer);
 		}
 
+		// Remap ids and copy used parametric types
+		uint32 SumParams = MakeIdRemapping(/*out*/ ParametricTypes, Used.KeptParametricTypes, Ids.GetParametricTypes());
+		CopyUsed(/* out */ KeptParametrics, Used.KeptParametricTypes, Ids.GetParametricTypes());
+
 		// Remap copied parametric types and copy parameters
 		TConstArrayView<FTypeId> AllParameters = Ids.GetParametricTypes().GetAllParameters();
+		KeptParameters.Reserve(SumParams);
 		for (FParametricType& KeptType : KeptParametrics)
 		{
 			FParameterIndexRange OldParameters = KeptType.Parameters;
