PlainProps: Initial drop of parsing a printed yaml document.

#rnx

diff --git a/Engine/Plugins/Experimental/PlainProps/Source/Private/PlainPropsBuild.cpp b/Engine/Plugins/Experimental/PlainProps/Source/Private/PlainPropsBuild.cpp
--- a/Engine/Plugins/Experimental/PlainProps/Source/Private/PlainPropsBuild.cpp
+++ b/Engine/Plugins/Experimental/PlainProps/Source/Private/PlainPropsBuild.cpp
@@ -118,6 +118,14 @@ FBuiltRange* CloneLeaves(FScratchAllocator& Scratch, uint64 Num, const void* InD
 
 //////////////////////////////////////////////////////////////////////////
 
+void FMemberBuilder::AddSuperStruct(FStructSchemaId SuperSchema, FBuiltStruct* SuperStruct)
+{
+	check(SuperStruct);
+	check(Members.IsEmpty());
+	Members.Emplace(FBuiltMember::MakeSuper(SuperSchema, SuperStruct));
+	check(IsSuper(Members[0].Schema.Type));
+}
+
 void FMemberBuilder::BuildSuperStruct(FScratchAllocator& Scratch, const FStructDeclaration& Super, const FDebugIds& Debug)
 {
 	// If we need to support dense substructs, we need access to struct declaration here 
diff --git a/Engine/Plugins/Experimental/PlainProps/Source/Private/PlainPropsInternalFormat.h b/Engine/Plugins/Experimental/PlainProps/Source/Private/PlainPropsInternalFormat.h
--- a/Engine/Plugins/Experimental/PlainProps/Source/Private/PlainPropsInternalFormat.h
+++ b/Engine/Plugins/Experimental/PlainProps/Source/Private/PlainPropsInternalFormat.h
@@ -135,6 +135,12 @@ inline uint32 CalculateSize(const FEnumSchema& In)
 	return Align(offsetof(FEnumSchema, Footer) + In.Num * sizeof(In.Footer[0]) + In.ExplicitConstants * In.Num * SizeOf(In.Width), alignof(FEnumSchema));
 }
 
+template<typename T>
+TConstArrayView<T> GetConstants(const FEnumSchema& In) requires (UnsignedIntegral<T> || Enumeration<T>)
+{
+	check(In.Width == ReflectLeaf<T>.Width);
+	return MakeConstArrayView<T>(reinterpret_cast<const T*>(In.Footer + In.Num), In.ExplicitConstants ? In.Num : 0);
+}
 //////////////////////////////////////////////////////////////////////////
 
 inline bool IsEnum(FMemberType Type)
diff --git a/Engine/Plugins/Experimental/PlainProps/Source/Private/PlainPropsInternalParse.h b/Engine/Plugins/Experimental/PlainProps/Source/Private/PlainPropsInternalParse.h
new file mode 100644
--- /dev/null
+++ b/Engine/Plugins/Experimental/PlainProps/Source/Private/PlainPropsInternalParse.h
@@ -0,0 +1,162 @@
+// Copyright Epic Games, Inc. All Rights Reserved.
+
+#pragma once
+
+#include "PlainPropsParse.h"
+#include "PlainPropsDeclare.h"
+#include "PlainPropsIndex.h"
+#include "Containers/Array.h"
+#include "Containers/ArrayView.h"
+#include "Containers/StringView.h"
+#include "Misc/Optional.h"
+#include "Templates/Tuple.h"
+#include "Templates/UniquePtr.h"
+
+namespace PlainProps
+{
+
+class FScratchAllocator;
+class FYamlTokenizer;
+struct FBuiltStruct;
+struct FEnumerator;
+
+///////////////////////////////////////////////////////////////////////////////
+
+enum class EToken : uint8
+{
+	Invalid,
+	BeginDocument,
+	EndDocument,
+	BeginStruct,
+	EndStruct,
+	BeginRange,
+	EndRange,
+	LeafId,
+	LeafValue,
+	Leaf,
+};
+
+struct FToken
+{ 
+	FToken() = default;
+
+	FToken(EToken InToken, uint16 InDepth) : Token(InToken), Depth(InDepth) { }
+
+	FToken(EToken InToken, uint16 InDepth, FUtf8StringView InView)
+		: Str(InView.GetData()), Len(InView.Len()), Token(InToken), Depth(InDepth)
+	{ }
+
+	inline FUtf8StringView Value() const { return FUtf8StringView(Str, Len); }
+
+	inline bool operator==(FToken O) const
+	{
+		return Token == O.Token && Depth == O.Depth && Value() == O.Value();
+	}
+
+	const UTF8CHAR*					Str		= nullptr;
+	int32							Len		= 0;
+	EToken							Token	= EToken::Invalid;
+	uint16							Depth	= 0;
+};
+
+///////////////////////////////////////////////////////////////////////////////
+
+struct FParsedMemberSchema
+{
+	FMemberType						Type;
+	FOptionalSchemaId				InnerSchema;
+	TArray<FMemberType>				InnerRangeTypes;
+};
+
+struct FParsedStructSchema
+{
+	FStructSchemaId					Id;
+	FOptionalStructSchemaId			DeclaredSuper;
+	FOptionalStructSchemaId			UsedSuper;
+	TConstArrayView<FMemberId>		MemberNames;	// FStructDeclaration.MemberOrder
+	TArray<FParsedMemberSchema>		MemberSchemas;	// Same size as declared member names
+};
+
+struct FParsedEnumSchema
+{
+	FEnumSchemaId					Id;
+	ELeafWidth						Width;
+	TConstArrayView<FEnumerator>	Enumerators;	// FEnumDeclaration.Enumerators
+};
+
+///////////////////////////////////////////////////////////////////////////////
+
+struct FParsedSchemas
+{
+	TArray<FParsedStructSchema>		Structs;		// Same size as number of parsed and declared structs
+	TArray<FParsedEnumSchema>		Enums;			// Same size as number of parsed and declared enums
+};
+
+///////////////////////////////////////////////////////////////////////////////
+
+class FTextIndexer : public TIdIndexer<FUtf8StringView>
+{
+public:
+	FOptionalStructSchemaId GetStructId(FTypeId Type) const
+	{
+		FSetElementId Id = Structs.FindId(Type);
+		if (Id.IsValidId())
+		{
+			return FStructSchemaId{ IntCastChecked<uint32>(Id.AsInteger()) };
+		}
+		return NoId;
+	}
+
+	FOptionalEnumSchemaId GetEnumId(FTypeId Type) const
+	{
+		FSetElementId Id = Enums.FindId(Type);
+		if (Id.IsValidId())
+		{
+			return FEnumSchemaId{ IntCastChecked<uint32>(Id.AsInteger()) };
+		}
+		return NoId;
+	}
+};
+
+///////////////////////////////////////////////////////////////////////////////
+
+struct FTypeTokens
+{
+	FUtf8StringView					Typename;
+	TArray<FUtf8StringView>			Scopes;
+	TArray<FTypeTokens>				Parameters;
+};
+
+TOptional<FTypeTokens> TokenizeType(FUtf8StringView String);
+
+FTypeId MakeType(const FTypeTokens& TypeTokens, FTextIndexer& Names);
+
+///////////////////////////////////////////////////////////////////////////////
+
+class FBatchParser
+{
+public:
+	FBatchParser(FYamlTokenizer& InTokenizer, FScratchAllocator& InScratch)
+		: Tokenizer(InTokenizer), Scratch(InScratch), Types(Names)
+	{ }
+
+	~FBatchParser() {};
+
+	void Parse(TArray64<uint8>& Out);
+
+private:
+	void ParseAll();
+	void ParseStructSchemas(TConstArrayView64<FToken> TokensView);
+	void ParseEnumSchemas(TConstArrayView64<FToken> TokensView);
+	void ParseObjects(TConstArrayView64<FToken> TokensView);
+	void Write(TArray64<uint8>& Out);
+
+	FYamlTokenizer&									Tokenizer;
+	FScratchAllocator&								Scratch;
+	FTextIndexer									Names;
+	FDeclarations									Types;
+	TArray<TTuple<FStructSchemaId, FBuiltStruct*>>	Objects;
+	FParsedSchemas									Schemas;
+};
+
+} // namespace PlainProps
diff --git a/Engine/Plugins/Experimental/PlainProps/Source/Private/PlainPropsInternalTest.h b/Engine/Plugins/Experimental/PlainProps/Source/Private/PlainPropsInternalTest.h
new file mode 100644
--- /dev/null
+++ b/Engine/Plugins/Experimental/PlainProps/Source/Private/PlainPropsInternalTest.h
@@ -0,0 +1,13 @@
+// Copyright Epic Games, Inc. All Rights Reserved.
+
+#pragma once
+
+#include "PlainPropsTypes.h"
+#include "Memory/MemoryFwd.h"
+
+namespace PlainProps
+{
+
+bool DiffSchemas(FMemoryView A, FMemoryView B);
+
+} // namespace PlainProps
diff --git a/Engine/Plugins/Experimental/PlainProps/Source/Private/PlainPropsInternalText.h b/Engine/Plugins/Experimental/PlainProps/Source/Private/PlainPropsInternalText.h
new file mode 100644
--- /dev/null
+++ b/Engine/Plugins/Experimental/PlainProps/Source/Private/PlainPropsInternalText.h
@@ -0,0 +1,46 @@
+// Copyright Epic Games, Inc. All Rights Reserved.
+
+#pragma once
+
+#include "PlainPropsTypes.h"
+#include "Containers/StringView.h"
+
+namespace PlainProps
+{
+
+struct FLiterals
+{
+	FAnsiStringView Structs = "Structs";
+	FAnsiStringView Enums = "Enums";
+	FAnsiStringView Objects = "Objects";
+
+	FAnsiStringView DeclaredSuper = "Super";
+	FAnsiStringView Width = "Width";
+	FAnsiStringView FlagMode = "FlagMode";
+	FAnsiStringView Constants = "Constants";
+	FAnsiStringView Members = "Members";
+
+	FAnsiStringView True = "true";
+	FAnsiStringView False = "false";
+
+	FAnsiStringView Super = "__super__";
+	FAnsiStringView Dynamic = "__dynamic__";
+
+	FAnsiStringView Ranges[9] = { 
+		("Uni"), ("i8"), ("u8"), ("i16"), ("u16"), ("i32"), ("u32"), ("i64"), ("u64") };
+
+	FAnsiStringView Leaves[8][4] = {
+		{("bool"),		("ERR_b16"),	("ERR_b32"),	("ERR_b64")},
+		{("i8"),		("i16"),		("i32"),		("i64")},
+		{("u8"),		("u16"),		("u32"),		("u64")},
+		{("ERR_fp8"),	("ERR_fp16"),	("float"),		("double")},
+		{("hex8"),		("hex16"),		("hex32"),		("hex64")},
+		{("enum8"),		("enum16"),		("enum32"),		("enum64")},
+		{("utf8"),		("utf16"),		("utf32"),		("ERR_utf64")},
+		{("ERR_oob"),	("ERR_oob"),	("ERR_oob"),	("ERR_oob")}};
+
+	FAnsiStringView Widths[4] = { ("u8"), ("u16"), ("u32"), ("u64") };
+};
+extern const FLiterals GLiterals;
+
+} // namespace PlainProps
diff --git a/Engine/Plugins/Experimental/PlainProps/Source/Private/PlainPropsParse.cpp b/Engine/Plugins/Experimental/PlainProps/Source/Private/PlainPropsParse.cpp
new file mode 100644
--- /dev/null
+++ b/Engine/Plugins/Experimental/PlainProps/Source/Private/PlainPropsParse.cpp
@@ -0,0 +1,1902 @@
+// Copyright Epic Games, Inc. All Rights Reserved.
+
+#include "PlainPropsParse.h"
+#include "Containers/Map.h"
+#include "Containers/StringConv.h"
+#include "Misc/AsciiSet.h"
+#include "Misc/StringBuilder.h"
+#include "PlainPropsBind.h"
+#include "PlainPropsBuild.h"
+#include "PlainPropsBuildSchema.h"
+#include "PlainPropsInternalBuild.h"
+#include "PlainPropsInternalParse.h"
+#include "PlainPropsInternalText.h"
+#include "PlainPropsPrint.h"
+#include "PlainPropsWrite.h"
+#include <charconv>
+
+namespace PlainProps
+{
+
+// specialization used by FTextIndexer/TIdIndexer<FUtf8StringView>
+template<>
+void AppendString(FUtf8StringBuilderBase& Out, const FUtf8StringView& Str)
+{
+	Out.Append(Str);
+}
+
+auto Print(FUtf8StringView View)
+{
+	return WriteToString<64>(FStringView(StringCast<TCHAR>(View.GetData(), View.Len())));
+}
+
+auto Print(FToken Token)
+{
+	return WriteToString<64>(FStringView(StringCast<TCHAR>(Token.Str, Token.Len)));
+}
+
+inline const FParsedStructSchema& ResolveStructSchema(const FParsedSchemas& Schemas, FStructSchemaId Id)
+{
+	check(Id.Idx < IntCastChecked<uint32>(Schemas.Structs.Num()));
+	return Schemas.Structs[Id.Idx];
+}
+
+inline const FParsedEnumSchema& ResolveEnumSchema(const FParsedSchemas& Schemas, FEnumSchemaId Id)
+{
+	check(Id.Idx < IntCastChecked<uint32>(Schemas.Enums.Num()));
+	return Schemas.Enums[Id.Idx];
+}
+
+///////////////////////////////////////////////////////////////////////////////
+
+static constexpr FUnpackedLeafType GLeaves[] = {
+	ReflectArithmetic<bool>,
+	ReflectArithmetic<int8>,
+	ReflectArithmetic<int16>,
+	ReflectArithmetic<int32>,
+	ReflectArithmetic<int64>,
+	ReflectArithmetic<uint8>,
+	ReflectArithmetic<uint16>,
+	ReflectArithmetic<uint32>,
+	ReflectArithmetic<uint64>,
+	ReflectArithmetic<float>,
+	ReflectArithmetic<double>,
+	FUnpackedLeafType(ELeafType::Hex, ELeafWidth::B8),
+	FUnpackedLeafType(ELeafType::Hex, ELeafWidth::B16),
+	FUnpackedLeafType(ELeafType::Hex, ELeafWidth::B32),
+	FUnpackedLeafType(ELeafType::Hex, ELeafWidth::B64),
+	FUnpackedLeafType(ELeafType::Enum, ELeafWidth::B8),
+	FUnpackedLeafType(ELeafType::Enum, ELeafWidth::B16),
+	FUnpackedLeafType(ELeafType::Enum, ELeafWidth::B32),
+	FUnpackedLeafType(ELeafType::Enum, ELeafWidth::B64),
+	ReflectArithmetic<char8_t>,
+	ReflectArithmetic<char16_t>,
+	ReflectArithmetic<char32_t>
+};
+
+template<>
+TOptional<ERangeSizeType> Parse(FUtf8StringView String)
+{
+	if (String.Len() >= 2)
+	{
+		constexpr uint8 N = sizeof(GLiterals.Ranges)/sizeof(GLiterals.Ranges[0]);
+		for (uint8 I = (String[0] == 'i'); I < N; I += 2)
+		{
+			if (String == GLiterals.Ranges[I])
+			{
+				return ERangeSizeType{ I };
+			}
+		}
+	}
+	return NullOpt;
+}
+
+template<>
+TOptional<FUnpackedLeafType> Parse(FUtf8StringView String)
+{
+	static int32 MaxLeafLen = 0;
+	static TMap<FUtf8StringView, FUnpackedLeafType> LeavesMap = [](int32& MaxLeafLen)
+	{
+		TMap<FUtf8StringView, FUnpackedLeafType> Map;
+		for (FUnpackedLeafType Leaf : GLeaves)
+		{
+			FAnsiStringView String = ToString(Leaf);
+			LeavesMap.Add(String, Leaf);
+			if (String.Len() > MaxLeafLen)
+			{
+				MaxLeafLen = String.Len();
+			}
+		}
+		return Map;
+	}(MaxLeafLen);
+	
+	if (String.Len() > MaxLeafLen)
+	{
+		return NullOpt;
+	}
+
+	if (FUnpackedLeafType* Leaf = LeavesMap.Find(String))
+	{
+		return *Leaf;
+	}
+	return NullOpt;
+}
+
+template<>
+TOptional<ELeafWidth> Parse(FUtf8StringView String)
+{
+	constexpr uint8 N = sizeof(GLiterals.Widths)/sizeof(GLiterals.Widths[0]);
+	for (uint8 I = 0; I < N; ++I)
+	{
+		if (String == GLiterals.Widths[I])
+		{
+			return ELeafWidth{ I };
+		}
+	}
+	return NullOpt;
+}
+
+///////////////////////////////////////////////////////////////////////////////
+
+template<Arithmetic T> 
+TOptional<T> Parse(FUtf8StringView String) requires (std::is_integral_v<T>)
+{
+	const char* first = reinterpret_cast<const char*>(String.GetData());
+	const char* last = first + String.Len();
+	T Out;
+	std::from_chars_result Result = std::from_chars(first, last, Out, 10);
+	return Result.ec == std::errc{} ? Out : TOptional<T>{ NullOpt };
+}
+
+template<Arithmetic T> 
+TOptional<T> Parse(FUtf8StringView String) requires (std::is_floating_point_v<T>)
+{
+	const char* first = reinterpret_cast<const char*>(String.GetData());
+	const char* last = first + String.Len();
+	T Out;
+	std::from_chars_result Result = std::from_chars(first, last, Out, std::chars_format::general);
+	return Result.ec == std::errc{} ? Out : TOptional<T>{ NullOpt };
+}
+
+template<>
+TOptional<bool> Parse(FUtf8StringView String)
+{
+	if (String == GLiterals.True)
+	{
+		return true;
+	}
+	if (String == GLiterals.False)
+	{
+		return false;
+	}
+	return NullOpt;
+}
+
+template<>
+TOptional<char8_t> Parse(FUtf8StringView String)
+{
+	if (String.Len() == 0)
+	{
+		return 0;
+	}
+	if (String.Len() == 1 && String[0] <= 127)
+	{
+		return static_cast<char8_t>(String[0]);
+	}
+	return NullOpt;
+}
+
+template<>
+TOptional<char16_t> Parse(FUtf8StringView String)
+{
+	if (String.Len() == 0)
+	{
+		return 0;
+	}
+	if (String.Len() == 1 && String[0] <= 127)
+	{
+		return static_cast<char16_t>(String[0]);
+	}
+	UTF16CHAR Dest;
+	if (UTF16CHAR* DestEnd = FPlatformString::Convert(&Dest, 1, String.GetData(), String.Len()))
+	{
+		check(DestEnd == &Dest + 1);
+		return static_cast<char16_t>(Dest);
+	}
+	return NullOpt;
+}
+
+template<>
+TOptional<char32_t> Parse(FUtf8StringView String)
+{
+	if (String.Len() == 0)
+	{
+		return 0;
+	}
+	if (String.Len() == 1 && String[0] <= 127)
+	{
+		return static_cast<char32_t>(String[0]);
+	}
+	// PP-TEXT: FPlatformString::Convert missing for UTF32CHAR in Parse()
+	checkf(false, TEXT("Missing conversion from UTF8CHAR to UTF32CHAR"));
+	return NullOpt;
+}
+
+///////////////////////////////////////////////////////////////////////////////
+
+constexpr FAsciiSet LinebreakSet("\r\n");
+constexpr FAsciiSet WhitespaceSet(" ");
+constexpr FAsciiSet TokenSet("'");
+constexpr FAsciiSet ScopeSet(".");
+constexpr FAsciiSet OpenRangeSet("(");
+constexpr FAsciiSet CloseRangeSet(")");
+constexpr FAsciiSet RangeSet(OpenRangeSet | CloseRangeSet);
+constexpr FAsciiSet OpenParameterSet("<[");
+constexpr FAsciiSet CloseParameterSet(">]");
+constexpr FAsciiSet NextParameterSet(",");
+constexpr FAsciiSet ParameterSet(OpenParameterSet | CloseParameterSet | NextParameterSet);
+
+constexpr bool PeekIsColon(FUtf8StringView View)			{ return *View.GetData() == ':'; }
+constexpr bool PeekIsToken(FUtf8StringView View)			{ return *View.GetData() == '\''; }
+constexpr bool IsRangeTokenLine(FUtf8StringView View)		{ return *View.GetData() == '-'; }
+constexpr bool IsEmptyOrCommentLine(FUtf8StringView View)	{ return View.IsEmpty() || *View.GetData() == '#'; }
+
+///////////////////////////////////////////////////////////////////////////////
+
+bool TokenizeType(FTypeTokens& Out, FUtf8StringView& String)
+{
+	check(String != GLiterals.Dynamic);
+
+	// tokenize scopes and concrete typename
+	FUtf8StringView Typename = FAsciiSet::FindPrefixWithout(String, ParameterSet);
+	if (int32 Len = Typename.Len())
+	{
+		FUtf8StringView Scope = FAsciiSet::FindPrefixWithout(Typename, ScopeSet);
+		while (Scope.Len() < Typename.Len())
+		{
+			Out.Scopes.Add(Scope);
+			Typename.RightChopInline(Scope.Len() + 1);
+			Scope = FAsciiSet::FindPrefixWithout(Typename, ScopeSet);
+		}
+		String.RightChopInline(Len);
+	}
+	Out.Typename = Typename;
+
+	// tokenize parameters for parametric types
+	UTF8CHAR OpenDelimiter{ };
+	while (String.Len() > 0)
+	{ 
+		UTF8CHAR C = String[0];
+
+		if (OpenParameterSet.Contains(C))
+		{
+			OpenDelimiter = C;
+			String.RightChopInline(1);
+			FTypeTokens Param;
+			if (!TokenizeType(Param, String))
+			{
+				return false;
+			}
+			Out.Parameters.Add(Param);
+		}
+		else if (CloseParameterSet.Contains(C))
+		{
+			// PP-TEXT: Remove hardcoded parameter delimiters when tweaking print format
+			if ((OpenDelimiter == '<' && C == '>') || (OpenDelimiter == '[' && C == ']'))
+			{
+				String.RightChopInline(1);
+			}
+			else if (OpenDelimiter)
+			{
+				checkf(false, TEXT("Mismatched delimiter '%s'"), *Print(String));
+				return false;
+			}
+			return true;
+		}
+		else if (C == ',')
+		{
+			if (OpenDelimiter)
+			{
+				String.RightChopInline(1);
+				FTypeTokens Param;
+				if (!TokenizeType(Param, String))
+				{
+					return false;
+				}
+				Out.Parameters.Add(Param);
+			}
+			else
+			{
+				return true;
+			}
+		}
+		else
+		{
+			checkf(false, TEXT("Invalid or unexpected type character '%s'"), *Print(String));
+			return false;
+		}
+	}
+	return true;
+}
+
+TOptional<FTypeTokens> TokenizeType(FUtf8StringView String)
+{
+	FTypeTokens Out;
+	if (TokenizeType(Out, String))
+	{
+		return MoveTemp(Out);
+	}
+	return NullOpt;
+}
+
+FTypeId MakeType(const FTypeTokens& TypeTokens, FTextIndexer& Names)
+{
+	// build parameter array for parametric types
+	TArray<FTypeId> Params;
+	if (int32 Num = TypeTokens.Parameters.Num())
+	{
+		Params.Reserve(Num);
+		for (const FTypeTokens& Parameter : TypeTokens.Parameters)
+		{
+			Params.Emplace(MakeType(Parameter, Names));
+		}
+	}
+
+	// build scopes and concrete typename
+	const FTypenameId TypenameId = Names.MakeTypename(TypeTokens.Typename);
+	FScopeId ScopeId = TypeTokens.Scopes.Num() ? Names.MakeScope(TypeTokens.Scopes[0]) : FScopeId{ NoId };
+	for (const FUtf8StringView& Scope : MakeArrayView(TypeTokens.Scopes).RightChop(1))
+	{
+		ScopeId = Names.NestScope(ScopeId, Scope);
+	}
+
+	// last create the full type
+	if (Params.Num())
+	{
+		if (TypeTokens.Typename.Len())
+		{
+			return Names.MakeParametricType({ ScopeId, TypenameId }, Params);
+		}
+		return Names.MakeAnonymousParametricType(Params);
+	}
+	return { ScopeId, TypenameId };
+}
+
+////////////////////////////////////////////////////////////////////////////////
+
+struct FLineInfo
+{
+	FUtf8StringView			View;
+	FUtf8StringView			FullView;
+	uint32					LineNumber		= 0;
+	int32					NumSpaces		= 0;
+	bool					HasRangeToken	= false;
+	bool					EndOfFile		= false;
+};
+
+class FYamlTokenizer
+{
+public:
+	FYamlTokenizer(FUtf8StringView InStringView);
+	~FYamlTokenizer();
+
+	bool HasMore() const { return ReadLine.EndOfFile == false || ScopeStack.Num() > 0; }
+	TOptional<FToken> GrabToken();
+
+	// PP-TEXT: GetLastError is a very temporary API while errors are unexpected
+	FUtf8StringView GetLastError() const;
+
+private:
+	enum class EScope : uint8
+	{
+		Document,
+		Struct,
+		Range
+	};
+
+	static constexpr EToken ScopeEndTokens[3] =
+	{
+		EToken::EndDocument,
+		EToken::EndStruct,
+		EToken::EndRange
+	};
+
+	struct FScopeInfo
+	{
+		EScope					Scope;
+		int32					NumSpaces		= 0;
+	};
+
+	TOptional<FToken>		 	GrabBeginDocument();
+	TOptional<FToken>			GrabTokenInternal();
+
+	void						PushScopeStack(EScope Scope);
+	FToken						PopScopeStack();
+
+	FLineInfo					ReadOneLine();
+	void						AdvanceReadLines();
+
+	void						SetLastError(FUtf8StringView Message);
+
+	FUtf8StringView								Text;
+	FUtf8StringView								ReadView;
+	TUtf8StringBuilder<256>						LastError;
+	TArray<FScopeInfo, TInlineAllocator<32>>	ScopeStack;
+	TOptional<FToken>							CachedToken;
+	FLineInfo									ReadLine;
+	FLineInfo									NextLine;
+	uint32										LineNumber = 1;
+};
+
+////////////////////////////////////////////////////////////////////////////////
+
+class FTokenReader
+{
+	const FToken* It = nullptr;
+	const FToken* End = nullptr;
+
+public:
+	FTokenReader(TConstArrayView64<FToken> Tokens)
+		: It(Tokens.GetData())
+		, End(Tokens.GetData() + Tokens.Num())
+	{}
+
+	bool HasMore() const							{ return It < End; }
+	const FToken* Peek() const						{ check(It < End); return It; }
+	[[nodiscard]] FToken GrabToken()				{ check(It < End); return *It++; }
+	[[nodiscard]] FToken GrabToken(EToken Token)	{ check(It->Token == Token); return GrabToken(); }
+};
+
+class FScopedTokenReader
+{
+public:
+	explicit FScopedTokenReader(FTokenReader& InTokens)
+		: Tokens(InTokens), Depth(InTokens.Peek()->Depth)
+	{ }
+
+	FScopedTokenReader(FTokenReader& InTokens, FToken ParentToken)
+		: Tokens(InTokens), Depth(ParentToken.Depth + 1)
+	{ }
+
+	bool HasMore() const							{ return Tokens.Peek()->Depth >= Depth; }
+	uint16 GetDepth() const							{ return Depth; }
+	FTokenReader& GetTokens()						{ return Tokens; }
+	const FToken* Peek() const						{ return Tokens.Peek(); }
+	FToken GrabToken()								{ check(HasMore()); return Tokens.GrabToken(); }
+	[[nodiscard]] FToken GrabToken(EToken Token)	{ return Tokens.GrabToken(Token); }
+
+private:
+	FTokenReader& Tokens;
+	uint16 Depth;
+};
+
+///////////////////////////////////////////////////////////////////////////////
+
+class FEnumSchemaParser
+{
+public:
+	FEnumSchemaParser(FTokenReader& InTokens, FTextIndexer& InNames, FDeclarations& InTypes)
+		: ScopedReader(InTokens)
+		, Names(InNames)
+		, Types(InTypes)
+		{ }
+
+	bool HasMore() const { return ScopedReader.HasMore(); }
+
+	bool ParseEnumSchema(FParsedEnumSchema& Out);
+
+private:
+	FParsedEnumSchema DeclareEnumSchema(
+		const FTypeTokens& TypeTokens,
+		EEnumMode Mode,
+		ELeafWidth Width,
+		TConstArrayView<FUtf8StringView> ConstantNames,
+		TConstArrayView<uint64> ConstantValues);
+
+	FScopedTokenReader	ScopedReader;
+	FTextIndexer&	  	Names;
+	FDeclarations&	  	Types;
+};
+
+////////////////////////////////////////////////////////////////////////////////
+
+class FStructSchemaParser
+{
+public:
+	FStructSchemaParser(
+		FTokenReader& InTokens, FTextIndexer& InNames, FDeclarations& InTypes, FParsedSchemas& InSchemas)
+		: ScopedReader(InTokens), Names(InNames), Types(InTypes), Schemas(InSchemas)
+	{ }
+
+	bool HasMore() const { return ScopedReader.HasMore(); }
+
+	bool ParseStructSchema(FParsedStructSchema& Out);
+
+private:
+	bool ParseMemberSchema(FParsedMemberSchema& Out, FUtf8StringView String);
+
+	FParsedStructSchema DeclareStructSchema(
+		const FTypeTokens& TypeTokens,
+		const TOptional<FTypeTokens>& ParsedSuper,
+		TConstArrayView<FUtf8StringView> MemberNames,
+		TConstArrayView<FUtf8StringView> MemberTypes);
+
+	FScopedTokenReader	ScopedReader;
+	FTextIndexer&	  	Names;
+	FDeclarations&	  	Types;
+	FParsedSchemas&		Schemas;
+};
+
+////////////////////////////////////////////////////////////////////////////////
+
+class FMemberParser
+{
+public:
+	FMemberParser(
+		FScratchAllocator& InScratch,
+		const FParsedStructSchema& Schema,
+		FScopedTokenReader InMemberTokens,
+		FTextIndexer& InNames,
+		FDeclarations& InTypes,
+		FParsedSchemas& InSchemas);
+
+	FBuiltStruct* ParseMembers(FStructSchemaId SchemaId);
+
+private:
+	bool HasMore() const { return MemberTokens.HasMore(); }
+	bool PeekKind(EMemberKind& MemberKind) const;
+	bool ParseAll();
+	bool ParseLeaf();
+	bool ParseStruct();
+	bool ParseRange();
+
+	bool AdvanceToNextMember();
+	bool ParseDynamicStructType(FOptionalStructSchemaId& Out);
+
+	[[nodiscard]] FBuiltStruct* BuildAndReset(FStructSchemaId Schema);
+
+	[[nodiscard]] FTypedRange	ParseLeaves(const FMemberSchema& RangeSchema, FScopedTokenReader RangeTokens);
+	[[nodiscard]] FTypedRange	ParseStructs(const FMemberSchema& RangeSchema, FScopedTokenReader RangeTokens);
+	[[nodiscard]] FTypedRange	ParseRanges(const FMemberSchema& RangeSchema, FScopedTokenReader RangeTokens);
+
+	[[nodiscard]] FTypedRange	ParseRangeInternal(const FMemberSchema& RangeSchema);
+
+	[[nodiscard]] FBuiltStruct*	ParseMembersInternal(
+		FStructType StructType,
+		FStructSchemaId Schema,
+		FScopedTokenReader InnerTokens);
+
+	[[nodiscard]] FTypedRange	ParseUnicodeLeafValueAsRange(
+		ERangeSizeType NumType,
+		FUnpackedLeafType Leaf,
+		FUtf8StringView String);
+
+	FScratchAllocator&						Scratch;
+	FMemberBuilder							MemberBuilder;
+	TConstArrayView<FMemberId>				MemberNames;
+	TConstArrayView<FParsedMemberSchema>	MemberSchemas;
+	FScopedTokenReader						MemberTokens;
+	FTextIndexer&							Names;
+	FDeclarations&							Types;
+	FParsedSchemas&							Schemas;
+	FOptionalStructSchemaId					UsedSuper;
+	uint32									TokenIdx = 0;
+	uint32								 	MemberIdx = 0;
+	bool									HasSuper = false;
+};
+
+///////////////////////////////////////////////////////////////////////////////
+
+FYamlTokenizer::FYamlTokenizer(FUtf8StringView InStringView)
+: Text(InStringView)
+, ReadView(InStringView)
+{
+	AdvanceReadLines();
+	AdvanceReadLines();
+}
+
+FYamlTokenizer::~FYamlTokenizer()
+{
+	checkf(!HasMore(), TEXT("FYamlTokenizer has ungrabbed tokens"));
+}
+
+static FToken StealToken(TOptional<FToken>& Opt)
+{
+	FToken Token = Opt.GetValue();
+	Opt.Reset();
+	return Token;
+}
+
+TOptional<FToken> FYamlTokenizer::GrabToken()
+{
+	if (ScopeStack.IsEmpty())
+	{
+		return GrabBeginDocument();
+	}
+
+	if (CachedToken.IsSet())
+	{
+		return StealToken(CachedToken);
+	}
+
+	// Pop stack to return end scope tokens
+	if (ReadLine.NumSpaces < ScopeStack.Last().NumSpaces || ReadLine.EndOfFile)
+	{
+		return PopScopeStack();
+	}
+
+	TOptional<FToken> Result = GrabTokenInternal();
+
+	AdvanceReadLines();
+
+	return Result;
+}
+
+FUtf8StringView FYamlTokenizer::GetLastError() const
+{
+	return LastError.ToString();
+}
+
+TOptional<FToken> FYamlTokenizer::GrabBeginDocument()
+{
+	if (ReadLine.EndOfFile)
+	{
+		return NullOpt;
+	}
+	FUtf8StringView View;
+	if (ReadLine.FullView.StartsWith("---"))
+	{
+		View = ReadLine.FullView;
+		AdvanceReadLines();
+	}
+	check(ScopeStack.IsEmpty());
+	ScopeStack.Emplace(EScope::Document);
+	return FToken{ EToken::BeginDocument, /*Depth*/ 0, View};
+}
+
+static FUtf8StringView GrabTokenFromLine(FLineInfo& Line)
+{
+	if (PeekIsToken(Line.View))
+	{
+		check(TokenSet.Contains(*Line.View.GetData()));
+		Line.View.RightChopInline(1);
+
+		const FUtf8StringView Token = FAsciiSet::FindPrefixWithout(Line.View, TokenSet);
+		Line.View.RightChopInline(Token.Len());
+
+		check(TokenSet.Contains(*Line.View.GetData()));
+		Line.View.RightChopInline(1);
+
+		const FUtf8StringView Skip = FAsciiSet::FindPrefixWith(Line.View, WhitespaceSet);
+		Line.View.RightChopInline(Skip.Len());
+
+		return Token;
+	}
+	return {};
+}
+
+static bool GrabColonFromLine(FLineInfo& Line)
+{
+	if (PeekIsColon(Line.View))
+	{
+		Line.View.RightChopInline(1);
+
+		const FUtf8StringView Skip = FAsciiSet::FindPrefixWith(Line.View, WhitespaceSet);
+		Line.View.RightChopInline(Skip.Len());
+
+		return true;
+	}
+	return false;
+}
+
+TOptional<FToken> FYamlTokenizer::GrabTokenInternal()
+{
+	const int32 TokenDepth = ScopeStack.Num();
+	TOptional<FToken> ParsedToken;
+	FUtf8StringView Error;
+
+	// Parse scalars
+	FUtf8StringView FirstToken = GrabTokenFromLine(ReadLine);
+	const bool HasColon = !FirstToken.IsEmpty() && GrabColonFromLine(ReadLine);
+	FUtf8StringView SecondToken = HasColon ? GrabTokenFromLine(ReadLine) : FUtf8StringView{};
+
+	// Leaf with id and value, value may be empty string
+	if (SecondToken.GetData() != nullptr)
+	{
+		if (FirstToken.Len() > 0)
+		{
+			ParsedToken.Emplace(EToken::LeafId, TokenDepth, FirstToken);
+			CachedToken.Emplace(EToken::LeafValue, TokenDepth, SecondToken);
+		}
+		else
+		{
+			Error = UTF8TEXTVIEW("EmptyLeafId");
+		}
+	}
+	// Leaf with value only
+	else if (FirstToken.GetData() != nullptr && !HasColon)
+	{
+		if (FirstToken.Len() > 0)
+		{
+			ParsedToken.Emplace(EToken::Leaf, TokenDepth, FirstToken);
+		}
+		else
+		{
+			Error = UTF8TEXTVIEW("EmptyValue");
+		}
+	}
+
+	// else we have a range or struct member with or without a name in FirstToken
+
+	// Empty struct
+	else if (*ReadLine.View.GetData() == '{')
+	{
+		ParsedToken.Emplace(EToken::BeginStruct, TokenDepth, FirstToken);
+		CachedToken.Emplace(EToken::EndStruct, TokenDepth, ReadLine.View);
+	}
+	// Empty range
+	else if (*ReadLine.View.GetData() == '[')
+	{
+		ParsedToken.Emplace(EToken::BeginRange, TokenDepth, FirstToken);
+		CachedToken.Emplace(EToken::EndRange, TokenDepth, ReadLine.View);
+	}
+	// Document end
+	else if (*ReadLine.View.GetData() == '.')
+	{
+		if (ReadLine.NumSpaces == 0 && ReadLine.View.StartsWith("..."))
+		{
+			ParsedToken.Emplace(PopScopeStack());
+		}
+		else
+		{
+			Error = UTF8TEXTVIEW("InvalidEndDocument");
+		}
+	}
+	// Non-empty range or struct
+	else if (NextLine.NumSpaces > ReadLine.NumSpaces)
+	{
+		if (*ReadLine.View.GetData() == '\n' || *ReadLine.View.GetData() == '#')
+		{
+			const EToken Token = NextLine.HasRangeToken ? EToken::BeginRange : EToken::BeginStruct;
+			const EScope Scope = NextLine.HasRangeToken ? EScope::Range : EScope::Struct;
+			ParsedToken.Emplace(Token, TokenDepth, FirstToken);
+			PushScopeStack(Scope);
+		}
+		else
+		{
+			Error = NextLine.HasRangeToken ? UTF8TEXTVIEW("InvalidRange") : UTF8TEXTVIEW("InvalidStruct");
+		}
+	}
+	else
+	{
+		Error = UTF8TEXTVIEW("UnknownToken");
+	}
+
+	check(ParsedToken.IsSet() == Error.IsEmpty());
+
+	SetLastError(Error);
+	return ParsedToken;
+}
+
+void FYamlTokenizer::PushScopeStack(EScope Scope)
+{
+	ScopeStack.Emplace(Scope, NextLine.NumSpaces);
+}
+
+FToken FYamlTokenizer::PopScopeStack()
+{
+	FScopeInfo Last = ScopeStack.Pop(EAllowShrinking::No);
+	return { ScopeEndTokens[static_cast<uint8>(Last.Scope)], IntCastChecked<uint16>(ScopeStack.Num()) };
+}
+
+FLineInfo FYamlTokenizer::ReadOneLine()
+{
+	if (ReadView.IsEmpty())
+	{
+		return { .View = ReadView, .FullView = ReadView, .LineNumber = LineNumber, .EndOfFile = true };
+	}
+
+	const FUtf8StringView LeadingSpaces = FAsciiSet::FindPrefixWith(ReadView, WhitespaceSet);
+	ReadView.RightChopInline(LeadingSpaces.Len());
+
+	const FUtf8StringView Line = FAsciiSet::FindPrefixWithout(ReadView, LinebreakSet);
+	ReadView.RightChopInline(Line.Len());
+
+	const FUtf8StringView TrailingLineBreaks = FAsciiSet::FindPrefixWith(ReadView, LinebreakSet);
+	ReadView.RightChopInline(TrailingLineBreaks.Len());
+
+	return { Line, Line, LineNumber++, LeadingSpaces.Len() };
+}
+
+void FYamlTokenizer::AdvanceReadLines()
+{
+	ReadLine = NextLine;
+	NextLine = ReadOneLine();
+
+	while (IsEmptyOrCommentLine(NextLine.View) && !NextLine.EndOfFile)
+	{
+		NextLine = ReadOneLine();
+	}
+
+	NextLine.HasRangeToken = IsRangeTokenLine(NextLine.View);
+	if (NextLine.HasRangeToken)
+	{
+		NextLine.View.RightChopInline(NextLine.HasRangeToken);
+		FUtf8StringView Skip = FAsciiSet::FindPrefixWith(NextLine.View, WhitespaceSet);
+		NextLine.View.RightChopInline(Skip.Len());
+	}
+}
+
+void FYamlTokenizer::SetLastError(FUtf8StringView Message)
+{
+	LastError.Reset();
+	LastError.Appendf("'%s' in line %u: '%s'", *Print(Message), ReadLine.LineNumber, *Print(ReadLine.FullView));
+}
+
+////////////////////////////////////////////////////////////////////////////////
+
+bool GrabToNextBeginStruct(FToken& Out, FScopedTokenReader& ScopedReader)
+{
+	while (ScopedReader.HasMore())
+	{
+		FToken Token = ScopedReader.GrabToken();
+		if (Token.Depth == ScopedReader.GetDepth() && Token.Token == EToken::BeginStruct)
+		{
+			Out = Token;
+			return true;
+		}
+	}
+	check(!ScopedReader.HasMore());
+	return false;
+}
+
+bool FEnumSchemaParser::ParseEnumSchema(FParsedEnumSchema& Out)
+{
+	FToken Token;
+	if (!GrabToNextBeginStruct(Token, ScopedReader))
+	{
+		return false;
+	}
+
+	TOptional<FTypeTokens> TypeTokens = TokenizeType(Token.Value());
+	if (!TypeTokens)
+	{
+		checkf(false, TEXT("Invalid enum schema typename '%s'"), *Print(Token));
+		return false;
+	}
+
+	// Parse Enum Properties
+	TOptional<EEnumMode> Mode;
+	TOptional<ELeafWidth> Width;
+	uint16 TokenDepth = Token.Depth + 1;
+	while (ScopedReader.Peek()->Depth == TokenDepth)
+	{
+		Token = ScopedReader.GrabToken();
+		if (Token.Token == EToken::LeafId)
+		{
+			if (Token.Value() == UTF8TEXTVIEW("FlagMode"))
+			{
+				Token = ScopedReader.GrabToken();
+				TOptional<bool> FlagMode = Parse<bool>(Token.Value());
+				if (FlagMode.IsSet())
+				{
+					Mode = FlagMode.GetValue() ? EEnumMode::Flag : EEnumMode::Flat;
+				}
+				else
+				{
+					checkf(false, TEXT("Unknown value '%s' for enum flag mode"), *Print(Token));
+				}
+			}
+			else if (Token.Value() == UTF8TEXTVIEW("Width"))
+			{
+				Token = ScopedReader.GrabToken();
+				Width = Parse<ELeafWidth>(Token.Value());
+				checkf(Width, TEXT("Unknown value '%s' for enum width"), *Print(Token));
+			}
+			else
+			{
+				checkf(false, TEXT("Unknown property '%s' for enum"), *Print(Token));
+				(void)ScopedReader.GrabToken();
+			}
+		}
+		else if (Token.Token == EToken::BeginRange)
+		{
+			checkf(Token.Value() == GLiterals.Constants,
+				TEXT("Invalid range '%s' for enum"), *Print(Token));
+			++TokenDepth;
+			break;
+		}
+		else
+		{
+			checkf(false, TEXT("Invalid token %d '%s' for enum"), Token.Token, *Print(Token));
+		}
+	}
+
+	// Parse Enum Constants
+	TArray<FUtf8StringView, TInlineAllocator<64>> ConstantNames;
+	TArray<uint64, TInlineAllocator<64>> ConstantValues;
+	while (ScopedReader.Peek()->Depth == TokenDepth)
+	{
+		Token = ScopedReader.GrabToken();
+		if (Token.Token == EToken::LeafId)
+		{
+			// uint64 Value;
+			FToken TokenValue = ScopedReader.GrabToken();
+			if (TOptional<uint64> ParsedValue = Parse<uint64>(TokenValue.Value()))
+			{
+				ConstantNames.Add(Token.Value());
+				ConstantValues.Add(ParsedValue.GetValue());
+			}
+			else
+			{
+				checkf(false, TEXT("Invalid value '%s' for enum constant '%s'"),
+					*Print(TokenValue), *Print(Token));
+			}
+		}
+		else
+		{
+			checkf(false, TEXT("Invalid token %d '%s' in enum constants"), Token.Token, *Print(Token));
+		}
+	}
+
+	Out = DeclareEnumSchema(TypeTokens.GetValue(), Mode.GetValue(), Width.GetValue(), ConstantNames, ConstantValues);
+	return true;
+}
+
+FParsedEnumSchema FEnumSchemaParser::DeclareEnumSchema(
+	const FTypeTokens& TypeTokens,
+	EEnumMode Mode,
+	ELeafWidth Width,
+	TConstArrayView<FUtf8StringView> ConstantNames,
+	TConstArrayView<uint64> ConstantValues)
+{
+	FTypeId Type = MakeType(TypeTokens, Names);
+	FEnumSchemaId Id = Names.IndexEnum(Type);
+
+	TArray<FEnumerator, TInlineAllocator<64>> Enumerators;
+	Enumerators.Reserve(ConstantNames.Num());
+	for (int32 I = 0, N = ConstantNames.Num(); I < N; ++I)
+	{
+		Enumerators.Emplace(Names.MakeName(ConstantNames[I]), ConstantValues[I]);
+	}
+
+	Types.DeclareEnum(Id, Type, Mode, Enumerators, EEnumAliases::Strip);
+
+	return { Id, Width, Types.Get(Id).GetEnumerators() };
+}
+
+///////////////////////////////////////////////////////////////////////////////
+
+bool FStructSchemaParser::ParseStructSchema(FParsedStructSchema& Out)
+{
+	FToken Token;
+	if (!GrabToNextBeginStruct(Token, ScopedReader))
+	{
+		return false;
+	}
+
+	TOptional<FTypeTokens> TypeTokens = TokenizeType(Token.Value());
+	if (!TypeTokens)
+	{
+		checkf(false, TEXT("Invalid struct schema typename '%s'"), *Print(Token));
+		return false;
+	}
+
+	// Parse Struct Properties
+	TOptional<FTypeTokens> ParsedSuper;
+	uint16 TokenDepth = Token.Depth + 1;
+	while (ScopedReader.Peek()->Depth == TokenDepth)
+	{
+		Token = ScopedReader.GrabToken();
+		if (Token.Token == EToken::LeafId)
+		{
+			if (Token.Value() == GLiterals.DeclaredSuper)
+			{
+				Token = ScopedReader.GrabToken();
+				ParsedSuper = TokenizeType(Token.Value());
+				checkf(ParsedSuper, TEXT("Invalid typename '%s' for declared super"), *Print(Token));
+			}
+			else
+			{
+				checkf(false, TEXT("Unknown property '%s' for struct"), *Print(Token));
+				(void)ScopedReader.GrabToken();
+			}
+		}
+		else if (Token.Token == EToken::BeginRange)
+		{
+			checkf(Token.Value() == GLiterals.Members,
+				TEXT("Invalid range '%s' for struct"), *Print(Token));
+			++TokenDepth;
+			break;
+		}
+		else
+		{
+			checkf(false, TEXT("Invalid token %d '%s' for struct"), Token.Token, *Print(Token));
+		}
+	}
+
+	// Parse Struct Members
+	TArray<FUtf8StringView, TInlineAllocator<64>> MemberNames;
+	TArray<FUtf8StringView, TInlineAllocator<64>> MemberTypes;
+	while (ScopedReader.Peek()->Depth == TokenDepth)
+	{
+		Token = ScopedReader.GrabToken();
+		if (Token.Token == EToken::LeafId)
+		{
+			MemberNames.Add(Token.Value());
+			Token = ScopedReader.GrabToken();
+			MemberTypes.Add(Token.Value());
+		}
+		else
+		{
+			checkf(false, TEXT("Invalid token %d '%s' in struct members"), Token.Token, *Print(Token));
+		}
+	}
+
+	Out = DeclareStructSchema(TypeTokens.GetValue(), ParsedSuper, MemberNames, MemberTypes);
+	return true;
+}
+
+FParsedStructSchema FStructSchemaParser::DeclareStructSchema(
+	const FTypeTokens& TypeTokens,
+	const TOptional<FTypeTokens>& ParsedSuper,
+	TConstArrayView<FUtf8StringView> MemberNames,
+	TConstArrayView<FUtf8StringView> MemberTypes)
+{
+	const FTypeId Type = MakeType(TypeTokens, Names);
+	const FStructSchemaId Id = Names.IndexStruct(Type);
+	const FOptionalStructSchemaId DeclaredSuper = ParsedSuper ?
+		Names.IndexStruct(MakeType(ParsedSuper.GetValue(), Names)) :
+		FOptionalStructSchemaId(NoId);
+
+	FOptionalStructSchemaId UsedSuper;
+
+	TArray<FMemberId, TInlineAllocator<64>> MemberIds;
+	TArray<FParsedMemberSchema> MemberSchemas;
+	MemberIds.Reserve(MemberNames.Num());
+	MemberSchemas.Reserve(MemberNames.Num());
+
+	const int32 HasSuper = DeclaredSuper && MemberNames.Num() && MemberNames[0] == GLiterals.Super;
+	if (HasSuper)
+	{
+		if (MemberTypes[0] != GLiterals.Dynamic)
+		{
+			TOptional<FTypeTokens> ParsedSuperMember = TokenizeType(MemberTypes[0]);
+			check(ParsedSuperMember);
+			FTypeId UsedSuperType = MakeType(ParsedSuperMember.GetValue(), Names);
+			UsedSuper = Names.IndexStruct(UsedSuperType);
+		}
+	}
+
+	for (int32 I = HasSuper, N = MemberNames.Num(); I < N; ++I)
+	{
+		FParsedMemberSchema MemberSchema;
+		if (ParseMemberSchema(MemberSchema, MemberTypes[I]))
+		{
+			MemberIds.Emplace(Names.MakeName(MemberNames[I]));
+			MemberSchemas.Emplace(MoveTemp(MemberSchema));
+		}
+		else
+		{
+			checkf(false, TEXT("Failed to parse member schema %d '%s'"), I - HasSuper, *Print(MemberTypes[I]));
+		}
+	}
+
+	Types.DeclareStruct(Id, Type, MemberIds, EMemberPresence::AllowSparse, DeclaredSuper);
+
+	return { Id, DeclaredSuper, UsedSuper, Types.Get(Id).GetMemberOrder(), MoveTemp(MemberSchemas) };
+}
+
+static TOptional<ERangeSizeType> GrabRangeSize(FUtf8StringView& String)
+{
+	if (String.Len() > 2 && OpenRangeSet.Contains(String[0]))
+	{
+		FUtf8StringView Remainder = String.RightChop(1);
+		FUtf8StringView Skip = FAsciiSet::FindPrefixWith(Remainder, WhitespaceSet | OpenRangeSet);
+		Remainder.RightChopInline(Skip.Len());
+		FUtf8StringView SizeStr = FAsciiSet::FindPrefixWithout(Remainder, WhitespaceSet | CloseRangeSet);
+		Remainder.RightChopInline(SizeStr.Len());
+		Skip = FAsciiSet::FindPrefixWith(Remainder, WhitespaceSet | CloseRangeSet);
+		Remainder.RightChopInline(Skip.Len());
+		if (TOptional<ERangeSizeType> MaxSize = Parse<ERangeSizeType>(SizeStr))
+		{
+			String = Remainder;
+			return MaxSize;
+		}
+	}
+	return NullOpt;
+}
+
+///////////////////////////////////////////////////////////////////////////////
+
+bool FStructSchemaParser::ParseMemberSchema(FParsedMemberSchema& Out, FUtf8StringView String)
+{
+	FUtf8StringView TypeOrLeaf = FAsciiSet::FindPrefixWithout(String, RangeSet);
+
+	if (TOptional<FUnpackedLeafType> Leaf = Parse<FUnpackedLeafType>(TypeOrLeaf))
+	{
+		Out.Type = FMemberType(Leaf.GetValue().Pack());
+	}
+	else if (TypeOrLeaf == GLiterals.Super)
+	{
+		checkf(TypeOrLeaf.Len() == String.Len(), TEXT("Super can't be used in a range"));
+		Out.Type = SuperStructType;
+		return true;
+	}
+	else if (TypeOrLeaf == GLiterals.Dynamic)
+	{
+		Out.Type = DynamicStructType;
+	}
+	else if (TOptional<FTypeTokens> TypeTokens = TokenizeType(TypeOrLeaf))
+	{
+		FTypeId Type = MakeType(TypeTokens.GetValue(), Names);
+		FOptionalEnumSchemaId Enum = Names.GetEnumId(Type);
+		if (Enum)
+		{
+			// All Enum types have already been parsed and are known before parsing structs and members
+			const FParsedEnumSchema& EnumSchema = ResolveEnumSchema(Schemas, Enum.Get());
+			Out.Type = FMemberType(ELeafType::Enum, EnumSchema.Width);
+			Out.InnerSchema = Enum.Get();
+		}
+		else
+		{
+			// PP-TEXT: Index all Struct types in a pass before parsing struct members?
+			// FOptionalStructSchemaId Struct = Names.GetStructId(Type);
+			// check(Struct);
+			// Out.InnerSchema = Struct.Get();
+			Out.Type = FMemberType(FStructType{ EMemberKind::Struct});
+			FStructSchemaId Struct = Names.IndexStruct(Type);
+			Out.InnerSchema = Struct;
+		}
+	}
+
+	if (TypeOrLeaf.Len() < String.Len())
+	{
+		FUtf8StringView RangeSizes = String.RightChop(TypeOrLeaf.Len());
+		FMemberType InnermostType = Out.Type;
+		if (TOptional<ERangeSizeType> RangeSize = GrabRangeSize(RangeSizes))
+		{
+			Out.Type = FMemberType(RangeSize.GetValue());
+			while ((RangeSize = GrabRangeSize(RangeSizes)))
+			{
+				Out.InnerRangeTypes.Emplace(RangeSize.GetValue());
+			}
+			Out.InnerRangeTypes.Add(InnermostType);
+		}
+		else
+		{
+			return false;
+		}
+	}
+	return true;
+}
+
+///////////////////////////////////////////////////////////////////////////////
+
+FMemberParser::FMemberParser(
+	FScratchAllocator& InScratch,
+	const FParsedStructSchema& Schema,
+	FScopedTokenReader InMemberTokens,
+	FTextIndexer& InNames,
+	FDeclarations& InTypes,
+	FParsedSchemas& InSchemas)
+: Scratch(InScratch)
+, MemberNames(Schema.MemberNames)
+, MemberSchemas(Schema.MemberSchemas)
+, MemberTokens(InMemberTokens)
+, Names(InNames)
+, Types(InTypes)
+, Schemas(InSchemas)
+, UsedSuper(Schema.UsedSuper)
+, HasSuper(Schema.DeclaredSuper)
+{
+}
+
+FBuiltStruct* FMemberParser::ParseMembers(FStructSchemaId SchemaId)
+{
+	ParseAll();
+	return BuildAndReset(SchemaId);
+}
+
+bool FMemberParser::PeekKind(EMemberKind& MemberKind) const
+{
+	switch (MemberTokens.Peek()->Token)
+	{
+		case EToken::BeginStruct:	MemberKind = EMemberKind::Struct;	return true;
+		case EToken::BeginRange:	MemberKind = EMemberKind::Range;	return true;
+		case EToken::LeafId:		MemberKind = EMemberKind::Leaf;		return true;
+		case EToken::Leaf:			MemberKind = EMemberKind::Leaf;		return true;
+	}
+	return false;
+}
+
+bool FMemberParser::ParseAll()
+{
+	while (HasMore())
+	{
+		EMemberKind Kind;
+		if (PeekKind(Kind))
+		{
+			switch (Kind)
+			{
+				case EMemberKind::Leaf:
+					verify(ParseLeaf());
+					break;
+				case EMemberKind::Struct:
+					verify(ParseStruct());
+					break;
+				case EMemberKind::Range:
+					verify(ParseRange());
+					break;
+			}
+		}
+		else
+		{
+			FToken Token = MemberTokens.GrabToken();
+			checkf(false, TEXT("Invalid member token %d '%s'"), Token.Token, *Print(Token))
+		}
+	}
+	return true;
+}
+
+template<typename T>
+void AddMember(FMemberBuilder& Members, FMemberId Id, FUtf8StringView String)
+{
+	if (TOptional<T> Value = Parse<T>(String))
+	{
+		Members.Add(Id, Value.GetValue());
+	}
+}
+
+template<typename T>
+void AddEnum(FMemberBuilder& Members, FMemberId Id, FEnumSchemaId Enum, FUtf8StringView String)
+{
+	if (TOptional<T> Value = Parse<T>(String))
+	{
+		Members.AddEnum(Id, Enum, Value.GetValue());
+	}
+}
+
+static bool IsUnicodeString(const FParsedMemberSchema& MemberSchema)
+{
+	return MemberSchema.Type.IsRange() &&
+		MemberSchema.InnerRangeTypes.Num() == 1 &&
+		MemberSchema.InnerRangeTypes[0].AsLeaf().Type == ELeafType::Unicode;
+}
+
+bool FMemberParser::ParseLeaf()
+{
+	if (!AdvanceToNextMember())
+	{
+		return false;
+	}
+
+	FToken Token = MemberTokens.GrabToken(EToken::LeafId);
+	Token = MemberTokens.GrabToken(EToken::LeafValue);
+	const FUtf8StringView String = Token.Value();
+
+	const FMemberId MemberId = MemberNames[MemberIdx];
+	const FParsedMemberSchema& MemberSchema = MemberSchemas[MemberIdx];
+
+	if (IsUnicodeString(MemberSchema))
+	{
+		FRangeType RangeType = MemberSchema.Type.AsRange();
+		FUnpackedLeafType Leaf = MemberSchema.InnerRangeTypes[0].AsLeaf();
+		MemberBuilder.AddRange(MemberId, ParseUnicodeLeafValueAsRange(RangeType.MaxSize, Leaf, String));
+		return true;
+	}
+
+	const FUnpackedLeafType Leaf = MemberSchema.Type.AsLeaf();
+	const FOptionalEnumSchemaId Enum = static_cast<FOptionalEnumSchemaId>(MemberSchema.InnerSchema);
+
+	switch (Leaf.Type)
+	{
+		case ELeafType::Bool:
+			AddMember<bool>(MemberBuilder, MemberId, String);
+			break;
+		case ELeafType::IntS:
+			switch (Leaf.Width)
+			{
+				case ELeafWidth::B8:	AddMember<int8 >(MemberBuilder, MemberId, String); break;
+				case ELeafWidth::B16:	AddMember<int16>(MemberBuilder, MemberId, String); break;
+				case ELeafWidth::B32:	AddMember<int32>(MemberBuilder, MemberId, String); break;
+				case ELeafWidth::B64:	AddMember<int64>(MemberBuilder, MemberId, String); break;
+			}
+			break;
+		case ELeafType::IntU:
+			switch (Leaf.Width)
+			{
+				case ELeafWidth::B8:	AddMember<uint8 >(MemberBuilder, MemberId, String);	break;
+				case ELeafWidth::B16:	AddMember<uint16>(MemberBuilder, MemberId, String);	break;
+				case ELeafWidth::B32:	AddMember<uint32>(MemberBuilder, MemberId, String);	break;
+				case ELeafWidth::B64:	AddMember<uint64>(MemberBuilder, MemberId, String);	break;
+			}
+			break;
+		case ELeafType::Float:
+			switch (Leaf.Width)
+			{
+				case ELeafWidth::B8:	check(false);										break;
+				case ELeafWidth::B16:	check(false);										break;
+				case ELeafWidth::B32:	AddMember<float >(MemberBuilder, MemberId, String);	break;
+				case ELeafWidth::B64:	AddMember<double>(MemberBuilder, MemberId, String);	break;
+			}
+			break;
+		case ELeafType::Hex:
+			check(Leaf.Type != ELeafType::Hex);
+			break;
+		case ELeafType::Enum:
+			switch (Leaf.Width)
+			{
+				case ELeafWidth::B8:	AddEnum<uint8 >(MemberBuilder, MemberId, Enum.Get(), String); break;
+				case ELeafWidth::B16:	AddEnum<uint16>(MemberBuilder, MemberId, Enum.Get(), String); break;
+				case ELeafWidth::B32:	AddEnum<uint32>(MemberBuilder, MemberId, Enum.Get(), String); break;
+				case ELeafWidth::B64:	AddEnum<uint64>(MemberBuilder, MemberId, Enum.Get(), String); break;
+			}
+			break;
+		case ELeafType::Unicode:
+			switch (Leaf.Width)
+			{
+				case ELeafWidth::B8:	AddMember<char8_t >(MemberBuilder, MemberId, String);	break;
+				case ELeafWidth::B16:	AddMember<char16_t>(MemberBuilder, MemberId, String);	break;
+				case ELeafWidth::B32:	AddMember<char32_t>(MemberBuilder, MemberId, String);	break;
+				case ELeafWidth::B64:	check(false);											break;
+			}
+			break;
+	}
+
+	return true;
+}
+
+bool FMemberParser::ParseStruct()
+{
+	if (HasSuper && MemberIdx == 0 && MemberTokens.Peek()->Value() == GLiterals.Super)
+	{
+		FToken Token = MemberTokens.GrabToken(EToken::BeginStruct);
+		if (!UsedSuper)
+		{
+			verify(ParseDynamicStructType(UsedSuper));
+		}
+
+		FBuiltStruct* Struct = ParseMembersInternal(
+			SuperStructType.AsStruct(), UsedSuper.Get(), { MemberTokens.GetTokens(), Token });
+		MemberBuilder.AddSuperStruct(UsedSuper.Get(), Struct);
+
+		Token = MemberTokens.GrabToken(EToken::EndStruct);
+		return true;
+	}
+
+	if (!AdvanceToNextMember())
+	{
+		return false;
+	}
+
+	const FMemberId MemberId = MemberNames[MemberIdx];
+	const FParsedMemberSchema& MemberSchema = MemberSchemas[MemberIdx];
+	const FStructType StructType = MemberSchema.Type.AsStruct();
+	check(!StructType.IsSuper);
+	check(MemberSchema.InnerSchema || StructType.IsDynamic);
+
+	FToken Token = MemberTokens.GrabToken(EToken::BeginStruct);
+
+	FOptionalStructSchemaId Schema = static_cast<FOptionalStructSchemaId>(MemberSchema.InnerSchema);
+	if (StructType.IsDynamic)
+	{
+		verify(ParseDynamicStructType(Schema));
+	}
+
+	FBuiltStruct* Struct = ParseMembersInternal(StructType, Schema.Get(), { MemberTokens.GetTokens(), Token });
+	MemberBuilder.AddStruct(MemberId, Schema.Get(), Struct);
+
+	Token = MemberTokens.GrabToken(EToken::EndStruct);
+	return true;
+}
+
+bool FMemberParser::ParseRange()
+{
+	if (!AdvanceToNextMember())
+	{
+		return false;
+	}
+
+	const FMemberId MemberId = MemberNames[MemberIdx];
+	const FParsedMemberSchema& MemberSchema = MemberSchemas[MemberIdx];
+
+	FMemberSchema RangeSchema
+	{
+		.Type = MemberSchema.Type,
+		.InnerRangeType = MemberSchema.InnerRangeTypes[0],
+		.NumInnerRanges = IntCastChecked<uint16>(MemberSchema.InnerRangeTypes.Num()),
+		.InnerSchema = MemberSchema.InnerSchema,
+		.NestedRangeTypes =
+			MemberSchema.InnerRangeTypes.Num() > 1 ?
+			MemberSchema.InnerRangeTypes.GetData() :
+			nullptr,
+	};
+
+	MemberBuilder.AddRange(MemberId, ParseRangeInternal(RangeSchema));
+	return true;
+}
+
+bool FMemberParser::AdvanceToNextMember()
+{
+	const FUtf8StringView MemberName = MemberTokens.Peek()->Value();
+	const FMemberId MemberId = Names.NameMember(MemberName);
+
+	for (uint16 Idx = MemberIdx, N = MemberNames.Num(); Idx < N; ++Idx)
+	{
+		if (MemberId == MemberNames[Idx])
+		{
+			MemberIdx = Idx;
+			return true;
+		}
+	}
+#if DO_CHECK
+	for (uint16 Idx = 0; Idx < MemberIdx; ++Idx)
+	{
+		if (MemberId == MemberNames[Idx])
+		{
+			checkf(false, TEXT("Member '%s' appeared in non-declared order"), *Print(MemberName));
+			return false;
+		}
+	}
+#endif
+
+	checkf(false, TEXT("Member '%s' not found"), *Print(MemberName));
+	return false;
+}
+
+bool FMemberParser::ParseDynamicStructType(FOptionalStructSchemaId& Out)
+{
+	FToken Token = MemberTokens.GrabToken(EToken::LeafId);
+	check(Token.Value() == GLiterals.Dynamic);
+
+	Token = MemberTokens.GrabToken(EToken::LeafValue);
+	const FUtf8StringView StructName = Token.Value();
+	TOptional<FTypeTokens> TypeTokens = TokenizeType(StructName);
+	check(TypeTokens);
+	FTypeId Type = MakeType(TypeTokens.GetValue(), Names);
+	FOptionalStructSchemaId Schema = Names.GetStructId(Type); 
+	if (Schema)
+	{
+		Out = Schema;
+		return true;
+	}
+	checkf(false, TEXT("Failed to parse dynamic struct schema '%s'"), *Print(Token));
+	return false;
+}
+
+FBuiltStruct* FMemberParser::BuildAndReset(FStructSchemaId Schema)
+{
+	const FStructDeclaration& StructDeclaration = Types.Get(static_cast<FStructSchemaId>(Schema));
+	return MemberBuilder.BuildAndReset(Scratch, StructDeclaration, Names);
+}
+
+template<Arithmetic T>
+[[nodiscard]] static FTypedRange ParseLeafRangeValues(
+	FScratchAllocator& Scratch,
+	ERangeSizeType MaxSize,
+	FScopedTokenReader& RangeTokens)
+{
+	TArray<T, TInlineAllocator64<64>> Values;
+	while (RangeTokens.HasMore())
+	{
+		FToken Token = RangeTokens.GrabToken();
+		if (TOptional<T> Value = Parse<T>(Token.Value()))
+		{
+			Values.Add(Value.GetValue());
+		}
+	}
+	return BuildLeafRange(Scratch, MaxSize, TConstArrayView<T>(Values));
+}
+
+template<Arithmetic T>
+[[nodiscard]] static FTypedRange ParseEnumRange(
+	FScratchAllocator& Scratch,
+	FEnumSchemaId Enum,
+	ERangeSizeType MaxSize,
+	FScopedTokenReader& RangeTokens)
+{
+	TArray<T, TInlineAllocator64<64>> Values;
+	while (RangeTokens.HasMore())
+	{
+		FToken Token = RangeTokens.GrabToken();
+		if (TOptional<T> Value = Parse<T>(Token.Value()))
+		{
+			Values.Add(Value.GetValue());
+		}
+	}
+	return BuildEnumRange(Scratch, Enum, MaxSize, TConstArrayView<T>(Values));
+}
+
+FTypedRange FMemberParser::ParseLeaves(const FMemberSchema& RangeSchema, FScopedTokenReader RangeTokens)
+{
+	const FOptionalEnumSchemaId Enum = static_cast<FOptionalEnumSchemaId>(RangeSchema.InnerSchema);
+	const FMemberType InnermostType = RangeSchema.GetInnerRangeTypes().Last();
+	const FUnpackedLeafType Leaf = InnermostType.AsLeaf();
+	ERangeSizeType NumType = RangeSchema.Type.AsRange().MaxSize;
+
+	FTypedRange TypedRange;
+	switch (Leaf.Type)
+	{
+		case ELeafType::Bool:
+			TypedRange = ParseLeafRangeValues<bool>(Scratch, NumType, RangeTokens);
+			break;
+		case ELeafType::IntS:
+			switch (Leaf.Width)
+			{
+				case ELeafWidth::B8:  TypedRange = ParseLeafRangeValues<int8 >(Scratch, NumType, RangeTokens); break;
+				case ELeafWidth::B16: TypedRange = ParseLeafRangeValues<int16>(Scratch, NumType, RangeTokens); break;
+				case ELeafWidth::B32: TypedRange = ParseLeafRangeValues<int32>(Scratch, NumType, RangeTokens); break;
+				case ELeafWidth::B64: TypedRange = ParseLeafRangeValues<int64>(Scratch, NumType, RangeTokens); break;
+			}
+			break;
+		case ELeafType::IntU:
+			switch (Leaf.Width)
+			{
+				case ELeafWidth::B8:  TypedRange = ParseLeafRangeValues<uint8 >(Scratch, NumType, RangeTokens); break;
+				case ELeafWidth::B16: TypedRange = ParseLeafRangeValues<uint16>(Scratch, NumType, RangeTokens); break;
+				case ELeafWidth::B32: TypedRange = ParseLeafRangeValues<uint32>(Scratch, NumType, RangeTokens); break;
+				case ELeafWidth::B64: TypedRange = ParseLeafRangeValues<uint64>(Scratch, NumType, RangeTokens); break;
+			}
+			break;
+		case ELeafType::Float:
+			if (Leaf.Width == ELeafWidth::B32)
+			{
+				TypedRange = ParseLeafRangeValues<float>(Scratch, NumType, RangeTokens);
+			}
+			else
+			{
+				check(Leaf.Width == ELeafWidth::B64);
+				TypedRange = ParseLeafRangeValues<double>(Scratch, NumType, RangeTokens);
+			}
+			break;
+		case ELeafType::Hex:
+			// PP-TEXT: Implement ParseHexRange()
+			check(Leaf.Type != ELeafType::Hex);
+			break;
+		case ELeafType::Enum:
+			switch (Leaf.Width)
+			{
+				case ELeafWidth::B8:	TypedRange = ParseEnumRange<uint8 >(Scratch, Enum.Get(), NumType, RangeTokens); break;
+				case ELeafWidth::B16:	TypedRange = ParseEnumRange<uint16>(Scratch, Enum.Get(), NumType, RangeTokens); break;
+				case ELeafWidth::B32:	TypedRange = ParseEnumRange<uint32>(Scratch, Enum.Get(), NumType, RangeTokens); break;
+				case ELeafWidth::B64:	TypedRange = ParseEnumRange<uint64>(Scratch, Enum.Get(), NumType, RangeTokens); break;
+			}
+			{
+				FTypedRange Range = BuildEnumRange(Scratch, Enum.Get(), NumType, TConstArrayView<unsigned char>({'a', 'b'}));
+			}
+			break;
+		case ELeafType::Unicode:
+			checkf(!RangeTokens.HasMore(), TEXT("Should have been handled by PrintUnicodeRangeAsLeafValue/ParseUnicodeLeafValueAsRange"));
+			TypedRange = ParseUnicodeLeafValueAsRange(NumType, Leaf, {});
+			break;
+	}
+
+	return TypedRange;
+}
+
+FTypedRange FMemberParser::ParseStructs(const FMemberSchema& RangeSchema, FScopedTokenReader RangeTokens)
+{
+	check(RangeSchema.NumInnerRanges == 1);
+	ERangeSizeType NumType = RangeSchema.Type.AsRange().MaxSize;
+	const FStructType StructType = RangeSchema.InnerRangeType.AsStruct();
+	check(!StructType.IsSuper);
+	check(RangeSchema.InnerSchema || StructType.IsDynamic);
+
+	TArray<FBuiltStruct*, TInlineAllocator64<64>> Structs;
+	FOptionalStructSchemaId Schema = static_cast<FOptionalStructSchemaId>(RangeSchema.InnerSchema);
+
+	if (StructType.IsDynamic)
+	{
+		check(!Schema);
+		while (RangeTokens.HasMore())
+		{
+			FToken Token = RangeTokens.GrabToken(EToken::BeginStruct);
+			FOptionalStructSchemaId DynamicSchema;
+			verify(ParseDynamicStructType(DynamicSchema));
+			
+			checkf(!Schema || Schema == DynamicSchema,
+				TEXT("Heterogeneous struct ranges have not been implemented yet"));
+			Schema = DynamicSchema;
+
+			FBuiltStruct* Struct = ParseMembersInternal(
+				DynamicStructType.AsStruct(), Schema.Get(), { RangeTokens.GetTokens(), Token });
+			Structs.Add(Struct);
+			Token = RangeTokens.GrabToken(EToken::EndStruct);
+		}
+	}
+	else
+	{
+		check(Schema);
+		while (RangeTokens.HasMore())
+		{
+			FToken Token = RangeTokens.GrabToken(EToken::BeginStruct);
+			FBuiltStruct* Struct = ParseMembersInternal(StructType, Schema.Get(), { RangeTokens.GetTokens(), Token });
+			Structs.Add(Struct);
+			Token = RangeTokens.GrabToken(EToken::EndStruct);
+		}
+	}
+
+	FMemberSchema BuildRangeSchema = Schema ?
+		MakeStructRangeSchema(NumType, Schema.Get()) :
+		MakeDynamicStructRangeSchema(NumType);
+
+	FTypedRange TypedRange { BuildRangeSchema };
+	if (int64 Num = Structs.Num())
+	{
+		TypedRange.Values = FBuiltRange::Create(Scratch, Structs.Num(), sizeof(FBuiltStruct*));
+		FMemory::Memcpy(TypedRange.Values->Data, Structs.GetData(), Structs.NumBytes());
+	}		
+
+	return TypedRange;
+}
+
+FTypedRange FMemberParser::ParseRanges(const FMemberSchema& RangeSchema, FScopedTokenReader RangeTokens)
+{
+	check(RangeSchema.NumInnerRanges > 1);
+	check(RangeSchema.NestedRangeTypes);
+
+	const ERangeSizeType NumType = RangeSchema.Type.AsRange().MaxSize;
+	const FMemberSchema InnerRangeSchema
+	{
+		.Type = RangeSchema.NestedRangeTypes[0],
+		.InnerRangeType = RangeSchema.NestedRangeTypes[1],
+		.NumInnerRanges = IntCastChecked<uint16>(RangeSchema.NumInnerRanges - 1),
+		.InnerSchema = RangeSchema.InnerSchema,
+		.NestedRangeTypes = RangeSchema.NumInnerRanges > 2 ? RangeSchema.NestedRangeTypes + 1 : nullptr,
+	};
+
+	TArray64<FBuiltRange*> Ranges;
+	while (RangeTokens.HasMore())
+	{
+		FTypedRange TypedRange = ParseRangeInternal(InnerRangeSchema);
+		Ranges.Add(TypedRange.Values);
+	}
+
+	FBuiltRange* Out = nullptr;
+	if (int64 Num = Ranges.Num())
+	{
+		Out = FBuiltRange::Create(Scratch, Num, sizeof(FBuiltRange*));
+		FMemory::Memcpy(Out->Data, Ranges.GetData(), Ranges.NumBytes());
+	}
+
+	return { MakeNestedRangeSchema(Scratch, NumType, InnerRangeSchema), Out };
+}
+
+static bool IsUnicodeString(const FMemberSchema& MemberSchema)
+{
+	return MemberSchema.Type.IsRange() &&
+		MemberSchema.NumInnerRanges == 1 &&
+		MemberSchema.InnerRangeType.AsLeaf().Type == ELeafType::Unicode;
+}
+
+FTypedRange FMemberParser::ParseRangeInternal(const FMemberSchema& RangeSchema)
+{
+	FTypedRange TypedRange{ };
+
+	FToken Token = MemberTokens.GrabToken();
+
+	if (Token.Token == EToken::Leaf)
+	{
+		check(IsUnicodeString(RangeSchema));
+		FUtf8StringView String = Token.Value();
+		FRangeType RangeType = RangeSchema.Type.AsRange();
+		FUnpackedLeafType Leaf = RangeSchema.InnerRangeType.AsLeaf();
+		return ParseUnicodeLeafValueAsRange(RangeType.MaxSize, Leaf, String);
+	}
+
+	check(Token.Token == EToken::BeginRange);
+	{
+		FScopedTokenReader RangeTokens(MemberTokens.GetTokens(), Token);
+		switch (RangeSchema.InnerRangeType.GetKind())
+		{
+			case EMemberKind::Leaf:		TypedRange = ParseLeaves(RangeSchema, RangeTokens);		break;
+			case EMemberKind::Struct:	TypedRange = ParseStructs(RangeSchema, RangeTokens);	break;
+			case EMemberKind::Range:	TypedRange = ParseRanges(RangeSchema, RangeTokens);		break;
+		}
+		check(!RangeTokens.HasMore());
+	}
+	Token = MemberTokens.GrabToken(EToken::EndRange);
+
+	return TypedRange;
+}
+
+FBuiltStruct* FMemberParser::ParseMembersInternal(
+	FStructType StructType,
+	FStructSchemaId Schema,
+	FScopedTokenReader InnerTokens)
+{
+	const FParsedStructSchema& ParsedSchema = ResolveStructSchema(Schemas, Schema);
+	FMemberParser MemberParser(Scratch, ParsedSchema, InnerTokens, Names, Types, Schemas);
+	return MemberParser.ParseMembers(Schema);
+}
+
+FTypedRange FMemberParser::ParseUnicodeLeafValueAsRange(
+	ERangeSizeType NumType,
+	FUnpackedLeafType Leaf,
+	FUtf8StringView String)
+{
+	if (Leaf.Width == ELeafWidth::B8)
+	{
+		const char8_t* Src = reinterpret_cast<const char8_t*>(String.GetData());
+		return BuildLeafRange(Scratch, NumType, MakeArrayView(Src, String.Len()));
+	}
+	else if (Leaf.Width == ELeafWidth::B16)
+	{
+		auto Utf16 = StringCast<UTF16CHAR>(String.GetData(), String.Len());
+		const char16_t* Src = reinterpret_cast<const char16_t*>(Utf16.Get());
+		return BuildLeafRange(Scratch, NumType, MakeArrayView(Src, Utf16.Length()));
+	}
+	else // ELeafWidth::B32
+	{
+		check(Leaf.Width == ELeafWidth::B32);
+		// PP-TEXT: StringCast<UTF32CHAR>/ToUtf32() missing in ParseUnicodeLeafValueAsRange()
+		// auto Utf32 = StringCast<UTF32CHAR>(String.GetData(), String.Len());
+		// const char32_t* Src = reinterpret_cast<const char32_t*>(Utf32.Get());
+		// return BuildLeafRange(Scratch, NumType, MakeArrayView(Src, Utf32.Length()));
+		checkf(false, TEXT("StringCast<UTF32CHAR>/ToUtf32() not implemented"));
+		return FTypedRange{ };
+	}
+}
+
+///////////////////////////////////////////////////////////////////////////////
+
+void FBatchParser::Parse(TArray64<uint8>& Out)
+{
+	ParseAll();
+	Write(Out);
+}
+
+void FBatchParser::ParseAll()
+{
+	int64 EnumsIdx = -1;
+	int64 StructsIdx = -1;
+	int64 ObjectsIdx = -1;
+
+	TArray64<FToken> Tokens;
+	while (Tokenizer.HasMore())
+	{
+		TOptional<FToken> Token = Tokenizer.GrabToken();
+		if (!Token)
+		{
+			checkf(false, TEXT("%s"), *Print(Tokenizer.GetLastError()));
+			continue;
+		}
+
+		if (Token->Depth == 1 && Token->Token == EToken::BeginRange)
+		{
+			if (Token->Value() == GLiterals.Structs)
+			{
+				StructsIdx = Tokens.Num();
+			}
+			else if (Token->Value() == GLiterals.Enums)
+			{
+				EnumsIdx = Tokens.Num();
+			}
+			else if (Token->Value() == GLiterals.Objects)
+			{
+				ObjectsIdx = Tokens.Num();
+			}
+		}
+
+		Tokens.Add(*Token);
+	}
+
+	checkf(EnumsIdx >= 0, TEXT("No '%.*hs' section found"), GLiterals.Structs.Len(), GLiterals.Structs.GetData());
+	checkf(StructsIdx >= 0, TEXT("No '%.*hs' section found"), GLiterals.Enums.Len(), GLiterals.Enums.GetData());
+	checkf(ObjectsIdx >= 0, TEXT("No '%.*hs' section found"), GLiterals.Objects.Len(), GLiterals.Objects.GetData());
+
+	TConstArrayView64<FToken> TokensView(Tokens);
+	ParseEnumSchemas(TokensView.RightChop(EnumsIdx));
+	ParseStructSchemas(TokensView.RightChop(StructsIdx));
+	ParseObjects(TokensView.RightChop(ObjectsIdx));
+}
+
+void FBatchParser::ParseStructSchemas(TConstArrayView64<FToken> TokensView)
+{
+	FTokenReader Tokens(TokensView);
+	verify(Tokens.GrabToken().Value() == GLiterals.Structs);
+
+	FStructSchemaParser Parser(Tokens, Names, Types, Schemas);
+	while (Parser.HasMore())
+	{
+		FParsedStructSchema ParsedStructSchema;
+		if (Parser.ParseStructSchema(ParsedStructSchema))
+		{
+			Schemas.Structs.SetNum(Types.GetStructs().Num());
+			Schemas.Structs[ParsedStructSchema.Id.Idx] = MoveTemp(ParsedStructSchema);
+		}
+	}
+}
+
+void FBatchParser::ParseEnumSchemas(TConstArrayView64<FToken> TokensView)
+{
+	FTokenReader Tokens(TokensView);
+	verify(Tokens.GrabToken().Value() == GLiterals.Enums);
+
+	FEnumSchemaParser Parser(Tokens, Names, Types);
+	while (Parser.HasMore())
+	{
+		FParsedEnumSchema ParsedEnumSchema;
+		if (Parser.ParseEnumSchema(ParsedEnumSchema))
+		{
+			check(ParsedEnumSchema.Id.Idx == Schemas.Enums.Num());
+			Schemas.Enums.Emplace(MoveTemp(ParsedEnumSchema));
+		}
+	}
+}
+
+void FBatchParser::ParseObjects(TConstArrayView64<FToken> TokensView)
+{
+	FTokenReader Tokens(TokensView);
+	verify(Tokens.GrabToken().Value() == GLiterals.Objects);
+
+	FScopedTokenReader ScopedReader(Tokens);
+	while (ScopedReader.HasMore())
+	{
+		FToken Token;
+		if (!GrabToNextBeginStruct(Token, ScopedReader))
+		{
+			break;
+		}
+
+		TOptional<FTypeTokens> TypeTokens = TokenizeType(Token.Value());
+		if (!TypeTokens)
+		{
+			checkf(false, TEXT("Invalid object struct schema '%s'"), *Print(Token));
+			continue;
+		}
+
+		const FTypeId Type = MakeType(TypeTokens.GetValue(), Names);
+		const FOptionalStructSchemaId StructId = Names.GetStructId(Type);
+
+		checkf(StructId, TEXT("Object struct schema '%s' not found"), *Print(Token));
+		const FParsedStructSchema& ParsedSchema = ResolveStructSchema(Schemas, StructId.Get());
+
+		{
+			FScopedTokenReader MemberTokens(Tokens, Token);
+			FMemberParser MemberParser(Scratch, ParsedSchema, MemberTokens, Names, Types, Schemas);
+			FBuiltStruct* BuiltStruct = MemberParser.ParseMembers(StructId.Get());
+			check(!MemberTokens.HasMore());
+			Objects.Emplace(StructId.Get(), BuiltStruct);
+		}
+
+		Token = Tokens.GrabToken(EToken::EndStruct);
+	}
+}
+
+class FParseBindings : public IStructBindIds
+{
+public:
+	virtual FStructSchemaId GetDeclId(FStructSchemaId BindId) const override final
+	{
+		checkf(false, TEXT("All struct ids should be declared, nothing is bound with different names for text"));
+		return BindId;
+	}
+};
+
+void FBatchParser::Write(TArray64<uint8>& Out)
+{
+	FParseBindings Bindings;
+	FSchemasBuilder SchemaBuilders(Types, Bindings, Scratch, ESchemaFormat::StableNames);
+	for (const TPair<FStructSchemaId, FBuiltStruct*>& Object : Objects)
+	{
+		SchemaBuilders.NoteStructAndMembers(Object.Key, *Object.Value);
+	}
+	FBuiltSchemas BuiltSchemas = SchemaBuilders.Build(); 
+
+	FWriter Writer(Names, Bindings, BuiltSchemas, ESchemaFormat::StableNames);
+	Writer.WriteSchemas(Out);
+}
+
+///////////////////////////////////////////////////////////////////////////////
+
+void ParseYamlBatch(TArray64<uint8>& OutBinary, FUtf8StringView Yaml)
+{
+	FScratchAllocator Scratch;
+	FYamlTokenizer YamlScanner(Yaml);
+	FBatchParser BatchParser(YamlScanner, Scratch);
+	BatchParser.Parse(OutBinary);
+}
+
+} // namespace PlainProps
diff --git a/Engine/Plugins/Experimental/PlainProps/Source/Private/PlainPropsPrint.cpp b/Engine/Plugins/Experimental/PlainProps/Source/Private/PlainPropsPrint.cpp
--- a/Engine/Plugins/Experimental/PlainProps/Source/Private/PlainPropsPrint.cpp
+++ b/Engine/Plugins/Experimental/PlainProps/Source/Private/PlainPropsPrint.cpp
@@ -8,6 +8,7 @@
 #include "PlainPropsInternalFormat.h"
 #include "PlainPropsInternalPrint.h"
 #include "PlainPropsInternalRead.h"
+#include "PlainPropsInternalText.h"
 #include "Misc/StringBuilder.h"
 
 namespace PlainProps
@@ -17,25 +18,7 @@ static constexpr bool PrintWithComments = true;
 
 //////////////////////////////////////////////////////////////////////////
 
-struct FLiterals
-{
-	FAnsiStringView Super = "__super__";
-	FAnsiStringView Dynamic = "__dynamic__";
-	FAnsiStringView Ranges[9] = { 
-		("Uni"), ("i8"), ("u8"), ("i16"), ("u16"), ("i32"), ("u32"), ("i64"), ("u64") };
-	FAnsiStringView Leaves[8][4] = {
-		{("bool"),		("ERR_b16"),	("ERR_b32"),	("ERR_b64")},
-		{("i8"),		("i16"),		("i32"),		("i64")},
-		{("u8"),		("u16"),		("u32"),		("u64")},
-		{("ERR_fp8"),	("ERR_fp16"),	("float"),		("double")},
-		{("hex8"),		("hex16"),		("hex32"),		("hex64")},
-		{("enum8"),		("enum16"),		("enum32"),		("enum64")},
-		{("utf8"),		("utf16"),		("utf32"),		("ERR_utf64")},
-		{("ERR_oob"),	("ERR_oob"),	("ERR_oob"),	("ERR_oob")}};
-};
-static const FLiterals GLiterals;
-
-//////////////////////////////////////////////////////////////////////////
+const FLiterals GLiterals;
 
 FAnsiStringView ToString(ERangeSizeType Width)
 {
@@ -47,6 +30,11 @@ FAnsiStringView ToString(FUnpackedLeafType Leaf)
 	return GLiterals.Leaves[(uint8)Leaf.Type][(uint8)Leaf.Width];
 }
 
+FAnsiStringView ToString(ELeafWidth Width)
+{
+	return GLiterals.Widths[(uint8)Width];
+}
+
 //////////////////////////////////////////////////////////////////////////
 
 void FIdIndexerBase::InitParameterNames()
@@ -67,19 +55,8 @@ void FIdIndexerBase::InitParameterNames()
 
 ///////////////////////////////////////////////////////////////////////////////
 
-template<typename IntType>
-static TConstArrayView<IntType> GetConstants(const FEnumSchema& Schema)
-{
-	check(Schema.Width == ReflectLeaf<IntType>.Width);
-	return MakeConstArrayView<IntType>(
-		reinterpret_cast<const IntType*>(Schema.Footer + Schema.Num),
-		Schema.ExplicitConstants ? Schema.Num : 0);
-}
-
-///////////////////////////////////////////////////////////////////////////////
-
 FUtf8StringView PrintMemberSchemaImpl(
-	FUtf8StringBuilderBase& Out,
+	FUtf8Builder& Out,
 	const FDebugIds& Names,
 	FMemberType Type,
 	FOptionalSchemaId InnerSchema = NoId,
@@ -114,21 +91,21 @@ FUtf8StringView PrintMemberSchemaImpl(
 
 	if (Type.IsRange())
 	{
-		Out.AppendChar('[');
+		Out.AppendChar('(');
 		Out.Append(ToString(Type.AsRange().MaxSize));
-		Out.AppendChar(']');
+		Out.AppendChar(')');
 		for (FMemberType MemberType : InnerRangeTypes.LeftChop(1))
 		{
-			Out.AppendChar('[');
+			Out.AppendChar('(');
 			Out.Append(ToString(MemberType.AsRange().MaxSize));
-			Out.AppendChar(']');
+			Out.AppendChar(')');
 		}
 	}
 	return Out.ToView();
 }
 
 void PrintMemberSchema(
-	FUtf8StringBuilderBase& Out,
+	FUtf8Builder& Out,
 	const FDebugIds& Names, 
 	FMemberType Type,
 	FOptionalSchemaId InnerSchema,
@@ -210,7 +187,7 @@ private:
 class FYamlBuilder
 {
 public:
-	FYamlBuilder(FUtf8StringBuilderBase& InStringBuilder);
+	FYamlBuilder(FUtf8Builder& InStringBuilder);
 	~FYamlBuilder();
 
 	void BeginDocument();
@@ -245,7 +222,7 @@ private:
 		bool IsInStruct = true;
 	};
 
-	FUtf8StringBuilderBase& Text;
+	FUtf8Builder& Text;
 	TArray<FStackInfo, TInlineAllocator<32>> Stack;
 	uint32 IndentLevel = 0;
 	bool IsNewLine = true;
@@ -277,7 +254,7 @@ private:
 
 	bool IsUnicodeString(const FRangeView& RangeView);
 	void PrintUnicodeLeafValue(const FLeafView& LeafView);
-	void PrintUnicodeRangeAsLeafValue(ERangeSizeType NumType, const FRangeView& RangeView);
+	void PrintUnicodeRangeAsLeaf(FOptionalMemberId Id, ERangeSizeType NumType, const FRangeView& RangeView);
 
 	FYamlBuilder& TextBuilder;
 	const FDebugIds& Names;
@@ -285,8 +262,8 @@ private:
 
 ///////////////////////////////////////////////////////////////////////////////
 
-void PrintBatchAsYamlDoc(
-	FUtf8StringBuilderBase& Out,
+void PrintYamlBatch(
+	FUtf8Builder& Out,
 	const FDebugIds& Names,
 	FReadBatchId Batch,
 	TConstArrayView<FStructView> Objects)
@@ -423,14 +400,14 @@ void FYamlBuilderDeleter::operator()(FYamlBuilder* YamlBuilder) const
 	delete YamlBuilder;
 }
 
-FYamlBuilderPtr MakeYamlBuilder(FUtf8StringBuilderBase& StringBuilder)
+FYamlBuilderPtr MakeYamlBuilder(FUtf8Builder& StringBuilder)
 {
 	return FYamlBuilderPtr(new FYamlBuilder(StringBuilder));
 }
 
 ///////////////////////////////////////////////////////////////////////////////
 
-FYamlBuilder::FYamlBuilder(FUtf8StringBuilderBase& InStringBuilder)
+FYamlBuilder::FYamlBuilder(FUtf8Builder& InStringBuilder)
 : Text(InStringBuilder)
 {
 	Stack.Emplace();
@@ -597,14 +574,14 @@ void FBatchPrinter::PrintSchemas(FReadBatchId Batch)
 
 void FBatchPrinter::PrintSchemas(const FSchemaBatch& Batch, FReadBatchId BatchId)
 {
-	TextBuilder.BeginRange("Structs");
+	TextBuilder.BeginRange(GLiterals.Structs);
 	for (const FStructSchema& Struct : GetStructSchemas(Batch))
 	{
 		PrintStructSchema(Struct, BatchId);
 	}
 	TextBuilder.EndRange();
 
-	TextBuilder.BeginRange("Enums");
+	TextBuilder.BeginRange(GLiterals.Enums);
 	for (const FEnumSchema& EnumSchema : GetEnumSchemas(Batch))
 	{
 		PrintEnumSchema(EnumSchema);
@@ -614,7 +591,7 @@ void FBatchPrinter::PrintSchemas(const FSchemaBatch& Batch, FReadBatchId BatchId
 
 void FBatchPrinter::PrintObjects(TConstArrayView<FStructView> Objects)
 {
-	TextBuilder.BeginRange("Objects");
+	TextBuilder.BeginRange(GLiterals.Objects);
 	for (FStructView Object : Objects)
 	{
 		FMemberPrinter(TextBuilder, Names).PrintMembers(Object);
@@ -623,7 +600,7 @@ void FBatchPrinter::PrintObjects(TConstArrayView<FStructView> Objects)
 }
 
 static FUtf8StringView PrintMemberSchema(
-	FUtf8StringBuilderBase& Out,
+	FUtf8Builder& Out,
 	const FDebugIds& Names,
 	const FMemberSchemaView& Schema)
 {
@@ -649,10 +626,10 @@ void FBatchPrinter::PrintStructSchema(const FStructSchema& Struct, FReadBatchId
 	if (Reader.HasSuper())
 	{
 		const FStructSchema& SuperSchema = Reader.GetSuper().Resolve();
-		TextBuilder.AddLeaf("Super", *FPrintId<128>(Names, SuperSchema.Type));
+		TextBuilder.AddLeaf(GLiterals.DeclaredSuper, *FPrintId<128>(Names, SuperSchema.Type));
 	}
 
-	TextBuilder.BeginRange("Members");
+	TextBuilder.BeginRange(GLiterals.Members);
 	while (Reader.HasMore())
 	{
 		TUtf8StringBuilder<256> Buf;
@@ -667,9 +644,10 @@ void FBatchPrinter::PrintStructSchema(const FStructSchema& Struct, FReadBatchId
 void FBatchPrinter::PrintEnumSchema(const FEnumSchema& Enum)
 {
 	TextBuilder.BeginStruct(*FPrintId<128>(Names, Enum.Type));
-	TextBuilder.AddLeaf("FlagMode", !!Enum.FlagMode);
+	TextBuilder.AddLeaf(GLiterals.FlagMode, !!Enum.FlagMode);
+	TextBuilder.AddLeaf(GLiterals.Width, ToString(Enum.Width));
 
-	TextBuilder.BeginRange("Constants");
+	TextBuilder.BeginRange(GLiterals.Constants);
 	TConstArrayView<FNameId> EnumNames = MakeConstArrayView(Enum.Footer, Enum.Num);
 	switch (Enum.Width)
 	{
@@ -806,8 +784,7 @@ void FMemberPrinter::PrintRange(FMemberId Id, ERangeSizeType NumType, const FRan
 {
 	if (IsUnicodeString(RangeView))
 	{
-		TextBuilder.AddLeafId(*FPrintId<128>(Names, Id));
-		PrintUnicodeRangeAsLeafValue(NumType, RangeView);
+		PrintUnicodeRangeAsLeaf(Id, NumType, RangeView);
 	}
 	else
 	{
@@ -852,7 +829,7 @@ void FMemberPrinter::PrintLeaves(FUnpackedLeafType Leaf, const FLeafRangeView& L
 		}
 		break;
 	case ELeafType::Hex:
-		// TODO: Implement AddLeaf(Hex)
+		// PP-TEXT: Implement AddLeaf(Hex)
 		check(Leaf.Type != ELeafType::Hex);
 		break;
 	case ELeafType::Enum:
@@ -865,7 +842,7 @@ void FMemberPrinter::PrintLeaves(FUnpackedLeafType Leaf, const FLeafRangeView& L
 		}
 		break;
 	case ELeafType::Unicode:
-		checkf(LeafRange.Num() == 0, TEXT("Should have been handled by PrintUnicodeRangeAsLeafValue"));
+		checkf(LeafRange.Num() == 0, TEXT("Should have been handled by PrintUnicodeRangeAsLeaf"));
 		break;
 	}
 }
@@ -885,8 +862,7 @@ void FMemberPrinter::PrintRanges(ERangeSizeType NumType, const FNestedRangeView&
 	{
 		if (IsUnicodeString(RangeView))
 		{
-			TextBuilder.AddLeafId({});
-			PrintUnicodeRangeAsLeafValue(NumType, RangeView);
+			PrintUnicodeRangeAsLeaf(NoId, NumType, RangeView);
 		}
 		else
 		{
@@ -1005,93 +981,72 @@ bool FMemberPrinter::IsUnicodeString(const FRangeView& RangeView)
 	return RangeView.Num() > 0 && Type.IsLeaf() && Type.AsLeaf().Type == ELeafType::Unicode;
 }
 
-void FMemberPrinter::PrintUnicodeLeafValue(const FLeafView& LeafView)
+template <typename CharType, typename LeafCharType>
+static void AddUnicodeLeaf(FYamlBuilder& TextBuilder, LeafCharType Char)
 {
-	if (LeafView.Leaf.Width == ELeafWidth::B8)
+	if (Char == 0)
 	{
-		UTF8CHAR Char8 = static_cast<UTF8CHAR>(LeafView.AsChar8());
-		check(Char8 <= 127);
-		if (Char8 == 0)
-		{
-			TextBuilder.AddLeafValue(FUtf8StringView{});
-		}
-		else
-		{
-			TextBuilder.AddLeafValue(Char8);
-		}
+		TextBuilder.AddLeafValue(FUtf8StringView{});
 	}
-	else if (LeafView.Leaf.Width == ELeafWidth::B16)
+	else if (Char <= 127)
 	{
-		UTF16CHAR Char16 = static_cast<UTF16CHAR>(LeafView.AsChar16());
-		if (Char16 == 0)
-		{
-			TextBuilder.AddLeafValue(FUtf8StringView{});
-		}
-		else if (Char16 <= 127)
-		{
-			TextBuilder.AddLeafValue(static_cast<UTF8CHAR>(Char16));
-		}
-		else
-		{
-			TextBuilder.AddLeafValue(FUtf8StringView(StringCast<UTF8CHAR>(&Char16, 1)));
-		}
+		TextBuilder.AddLeafValue(static_cast<UTF8CHAR>(Char));
 	}
-	else // ELeafWidth::B32
+	else
 	{
-		check(LeafView.Leaf.Width != ELeafWidth::B64);
-		UTF32CHAR Char32 = static_cast<UTF32CHAR>(LeafView.AsChar32());
-		if (Char32 == 0)
-		{
-			TextBuilder.AddLeafValue(FUtf8StringView{});
-		}
-		else if (Char32 <= 127)
-		{
-			TextBuilder.AddLeafValue(static_cast<UTF8CHAR>(Char32));
-		}
-		else
-		{
-			TextBuilder.AddLeafValue(FUtf8StringView(StringCast<UTF8CHAR>(&Char32, 1)));
-		}
+		CharType Temp = static_cast<CharType>(Char);
+		TextBuilder.AddLeafValue(FUtf8StringView(StringCast<UTF8CHAR>(&Temp, 1)));
 	}
 }
 
-void FMemberPrinter::PrintUnicodeRangeAsLeafValue(ERangeSizeType NumType, const FRangeView& RangeView)
+void FMemberPrinter::PrintUnicodeLeafValue(const FLeafView& LeafView)
 {
-	check(IsUnicodeString(RangeView));
-
-	const FLeafRangeView LeafRange = RangeView.AsLeaves();
-	FUnpackedLeafType Leaf = RangeView.Schema.ItemType.AsLeaf();
+	switch (LeafView.Leaf.Width)
+	{
+		case ELeafWidth::B8:	AddUnicodeLeaf<UTF8CHAR >(TextBuilder, LeafView.AsChar8());		break;
+		case ELeafWidth::B16:	AddUnicodeLeaf<UTF16CHAR>(TextBuilder, LeafView.AsChar16());	break;
+		case ELeafWidth::B32:	AddUnicodeLeaf<UTF32CHAR>(TextBuilder, LeafView.AsChar32());	break;
+		case ELeafWidth::B64:	check(false);													break;
+	};
+}
 
-	if (Leaf.Width == ELeafWidth::B8)
+template <typename CharType, typename RangeCharType>
+static void AddUnicodeRangeLeaf(FYamlBuilder& TextBuilder, FOptionalMemberId Id, TRangeView<RangeCharType> Utf)
+{
+	check(Utf.Num() > 0);
+	const CharType* Src = reinterpret_cast<const CharType*>(Utf.begin());
+	int32 Len = IntCastChecked<int32>(Utf.Num());
+	Len -= (Utf[Len-1] == CharType(0)); // strip null terminator (for utf8->utf8 where StringCast is pass through)
+	if (Id)
 	{
-		TRangeView<char8_t> Utf8 = LeafRange.AsUtf8();
-		check(Utf8.Num() > 0);
-		const UTF8CHAR* Src = reinterpret_cast<const UTF8CHAR*>(Utf8.begin());
-		int32 Len = IntCastChecked<int32>(Utf8.Num());
-		// explicitly strip null terminator, no StringCast required
-		Len -= (Utf8[Len-1] == UTF8CHAR(0));
-		TextBuilder.AddLeafValue(FUtf8StringView(Src, Len));
+		TextBuilder.AddLeafValue(FUtf8StringView(StringCast<UTF8CHAR>(Src, Len)));
 	}
-	else if (Leaf.Width == ELeafWidth::B16)
+	else
 	{
-		TRangeView<char16_t> Utf16 = LeafRange.AsUtf16();
-		check(Utf16.Num() > 0);
-		const UTF16CHAR* Src = reinterpret_cast<const UTF16CHAR*>(Utf16.begin());
-		const int32 Len = IntCastChecked<int32>(Utf16.Num());
-		// null terminator is excluded by length of StringCast result
-		TextBuilder.AddLeafValue(FUtf8StringView(StringCast<UTF8CHAR>(Src, Len)));
+		TextBuilder.AddLeaf(FUtf8StringView(StringCast<UTF8CHAR>(Src, Len)));
 	}
-	else // ELeafWidth::B32
+}
+
+void FMemberPrinter::PrintUnicodeRangeAsLeaf(FOptionalMemberId Id, ERangeSizeType NumType, const FRangeView& RangeView)
+{
+	check(IsUnicodeString(RangeView));
+
+	if (Id)
 	{
-		check(Leaf.Width != ELeafWidth::B64);
-		TRangeView<char32_t> Utf32 = LeafRange.AsUtf32();
-		check(Utf32.Num() > 0);
-		const UTF32CHAR* Src = reinterpret_cast<const UTF32CHAR*>(Utf32.begin());
-		const int32 Len = IntCastChecked<int32>(Utf32.Num());
-		// null terminator is excluded by length of StringCast result
-		TextBuilder.AddLeafValue(FUtf8StringView(StringCast<UTF8CHAR>(Src, Len)));
+		TextBuilder.AddLeafId(*FPrintId<128>(Names, Id));
 	}
 
+	const FLeafRangeView LeafRange = RangeView.AsLeaves();
+	const FUnpackedLeafType Leaf = RangeView.Schema.ItemType.AsLeaf();
+
+	switch (Leaf.Width)
+	{
+		case ELeafWidth::B8:	AddUnicodeRangeLeaf<UTF8CHAR >(TextBuilder, Id, LeafRange.AsUtf8());	break;
+		case ELeafWidth::B16:	AddUnicodeRangeLeaf<UTF16CHAR>(TextBuilder, Id, LeafRange.AsUtf16());	break;
+		case ELeafWidth::B32:	AddUnicodeRangeLeaf<UTF32CHAR>(TextBuilder, Id, LeafRange.AsUtf32());	break;
+		case ELeafWidth::B64:	check(false);															break;
+	};
+
 	if constexpr(PrintWithComments)
 	{
 		AddRangeSchemaComment(TextBuilder, Names, FMemberType(NumType), RangeView.Schema);
@@ -1100,12 +1055,12 @@ void FMemberPrinter::PrintUnicodeRangeAsLeafValue(ERangeSizeType NumType, const
 
 ///////////////////////////////////////////////////////////////////////////////
 
-void FDebugIds::AppendString(FUtf8StringBuilderBase& Out, FMemberId Name) const
+void FDebugIds::AppendString(FUtf8Builder& Out, FMemberId Name) const
 {
 	AppendString(Out, Name.Id);
 }
 
-void FDebugIds::AppendString(FUtf8StringBuilderBase& Out, FOptionalMemberId Name) const
+void FDebugIds::AppendString(FUtf8Builder& Out, FOptionalMemberId Name) const
 {
 	if (Name)
 	{
@@ -1117,7 +1072,7 @@ void FDebugIds::AppendString(FUtf8StringBuilderBase& Out, FOptionalMemberId Name
 	}
 }
 
-void FDebugIds::AppendString(FUtf8StringBuilderBase& Out, FScopeId Scope) const
+void FDebugIds::AppendString(FUtf8Builder& Out, FScopeId Scope) const
 {
 	if (Scope.IsFlat())
 	{
@@ -1132,7 +1087,7 @@ void FDebugIds::AppendString(FUtf8StringBuilderBase& Out, FScopeId Scope) const
 	}
 }
 
-void FDebugIds::AppendString(FUtf8StringBuilderBase& Out, FTypenameId Typename) const
+void FDebugIds::AppendString(FUtf8Builder& Out, FTypenameId Typename) const
 {
 	if (Typename.IsConcrete())
 	{
@@ -1162,7 +1117,7 @@ void FDebugIds::AppendString(FUtf8StringBuilderBase& Out, FTypenameId Typename)
 	}
 }
 
-void FDebugIds::AppendString(FUtf8StringBuilderBase& Out, FTypeId Type) const
+void FDebugIds::AppendString(FUtf8Builder& Out, FTypeId Type) const
 {
 	if (Type.Scope)
 	{
@@ -1172,12 +1127,12 @@ void FDebugIds::AppendString(FUtf8StringBuilderBase& Out, FTypeId Type) const
 	AppendString(Out, Type.Name);
 }
 
-void FDebugIds::AppendString(FUtf8StringBuilderBase& Out, FEnumSchemaId Name) const
+void FDebugIds::AppendString(FUtf8Builder& Out, FEnumSchemaId Name) const
 {
 	AppendString(Out, Resolve(Name));
 }
 
-void FDebugIds::AppendString(FUtf8StringBuilderBase& Out, FStructSchemaId Name) const
+void FDebugIds::AppendString(FUtf8Builder& Out, FStructSchemaId Name) const
 {
 	AppendString(Out, Resolve(Name));
 }
diff --git a/Engine/Plugins/Experimental/PlainProps/Source/Public/PlainPropsBuild.h b/Engine/Plugins/Experimental/PlainProps/Source/Public/PlainPropsBuild.h
--- a/Engine/Plugins/Experimental/PlainProps/Source/Public/PlainPropsBuild.h
+++ b/Engine/Plugins/Experimental/PlainProps/Source/Public/PlainPropsBuild.h
@@ -150,7 +150,7 @@ FMemberSchema MakeLeafRangeSchema()
 }
 
 template<Enumeration T>
-FMemberSchema MakeEnumRangeSchema(ERangeSizeType MaxSize, FEnumSchemaId Schema)
+FMemberSchema MakeEnumRangeSchema(FEnumSchemaId Schema, ERangeSizeType MaxSize)
 {
 	return { FMemberType(MaxSize), ReflectEnum<T>.Pack(), 1, Schema, nullptr };
 }
@@ -158,11 +158,18 @@ FMemberSchema MakeEnumRangeSchema(ERangeSizeType MaxSize, FEnumSchemaId Schema)
 template<Enumeration T, typename SizeType>
 FMemberSchema MakeEnumRangeSchema(FEnumSchemaId Schema)
 {
-	return MakeEnumRangeSchema<T>(RangeSizeOf(SizeType{}), Schema);
+	return MakeEnumRangeSchema<T>(Schema, RangeSizeOf(SizeType{}));
 }
 
-inline constexpr FMemberType DefaultStructType = FMemberType(FStructType{EMemberKind::Struct, /* IsDynamic */ 0, /* IsSuper */ 0});
-inline constexpr FMemberType SuperStructType =	 FMemberType(FStructType{EMemberKind::Struct, /* IsDynamic */ 0, /* IsSuper */ 1});
+template<UnsignedIntegral T>
+FMemberSchema MakeEnumRangeSchema(FEnumSchemaId Schema, ERangeSizeType MaxSize)
+{
+	return { FMemberType(MaxSize), FMemberType(ELeafType::Enum, ReflectArithmetic<T>.Width), 1, Schema, nullptr };
+}
+
+inline constexpr FMemberType DefaultStructType =	FMemberType(FStructType{EMemberKind::Struct, /* IsDynamic */ 0, /* IsSuper */ 0});
+inline constexpr FMemberType SuperStructType =		FMemberType(FStructType{EMemberKind::Struct, /* IsDynamic */ 0, /* IsSuper */ 1});
+inline constexpr FMemberType DynamicStructType =	FMemberType(FStructType{EMemberKind::Struct, /* IsDynamic */ 1, /* IsSuper */ 0});
 
 
 inline FMemberSchema MakeStructRangeSchema(ERangeSizeType SizeType, FStructSchemaId Schema)
@@ -170,24 +177,40 @@ inline FMemberSchema MakeStructRangeSchema(ERangeSizeType SizeType, FStructSchem
 	return { FMemberType(SizeType), DefaultStructType, 1, Schema, nullptr };
 }
 
+inline FMemberSchema MakeDynamicStructRangeSchema(ERangeSizeType SizeType)
+{
+	return { FMemberType(SizeType), DynamicStructType, 1, NoId, nullptr };
+}
+
 PLAINPROPS_API FMemberSchema MakeNestedRangeSchema(FScratchAllocator& Scratch, ERangeSizeType SizeType, FMemberSchema InnerRangeSchema);
 
 //////////////////////////////////////////////////////////////////////////
 
 [[nodiscard]] PLAINPROPS_API FBuiltRange* CloneLeaves(FScratchAllocator& Scratch, uint64 Num, const void* Data, SIZE_T LeafSize);
 
-template<LeafType T, typename SizeType>
-[[nodiscard]] FTypedRange BuildLeafRange(FScratchAllocator& Scratch, const T* Values, SizeType Num)
+template<Arithmetic T>
+[[nodiscard]] FTypedRange BuildLeafRange(FScratchAllocator& Scratch, ERangeSizeType SizeType, const T* Values, uint64 Num)
 {
 	// todo: detect invalid floats
-	return { MakeLeafRangeSchema<T, SizeType>(), CloneLeaves(Scratch, Num, Values, sizeof(T)) };
+	return { MakeLeafRangeSchema<T>(SizeType), CloneLeaves(Scratch, Num, Values, sizeof(T)) };
+}
+
+template<Arithmetic T, typename SizeType>
+[[nodiscard]] FTypedRange BuildLeafRange(FScratchAllocator& Scratch, const T* Values, SizeType Num)
+{
+	return BuildLeafRange(Scratch, RangeSizeOf(SizeType{}), Values, Num);
+}
+
+template<Arithmetic T>
+[[nodiscard]] FTypedRange BuildLeafRange(FScratchAllocator& Scratch, ERangeSizeType SizeType, TConstArrayView<T> Values)
+{
+	return BuildLeafRange(Scratch, SizeType, Values.GetData(), Values.Num());
 }
 
-template<LeafType T, typename SizeType>
+template<Arithmetic T, typename SizeType>
 [[nodiscard]] FTypedRange BuildLeafRange(FScratchAllocator& Scratch, TConstArrayView<T, SizeType> Values)
 {
-	// todo: detect invalid floats
-	return { MakeLeafRangeSchema<T, SizeType>(), CloneLeaves(Scratch, Values.Num(), Values.GetData(), sizeof(T)) };
+	return BuildLeafRange(Scratch, RangeSizeOf(SizeType{}), Values.GetData(), Values.Num());
 }
 
 template<Enumeration T, typename SizeType>
@@ -196,6 +219,12 @@ template<Enumeration T, typename SizeType>
 	return { MakeEnumRangeSchema<T, SizeType>(Enum), CloneLeaves(Scratch, Values.Num(), Values.GetData(), sizeof(T)) };
 }
 
+template<UnsignedIntegral T>
+[[nodiscard]] FTypedRange BuildEnumRange(FScratchAllocator& Scratch, FEnumSchemaId Enum, ERangeSizeType SizeType, TConstArrayView<T> Values)
+{
+	return { MakeEnumRangeSchema<T>(Enum, SizeType), CloneLeaves(Scratch, Values.Num(), Values.GetData(), sizeof(T)) };
+}
+
 [[nodiscard]] inline FTypedRange MakeStructRange(FStructSchemaId Schema, ERangeSizeType SizeType, FBuiltRange* Values )
 {
 	return { MakeStructRangeSchema(SizeType, Schema), Values };
@@ -248,20 +277,26 @@ public:
 		AddLeaf(Name, ReflectEnum<T>, ToOptional(Schema), ValueCast(static_cast<__underlying_type(T)>(Value)));
 	}
 	
-	void AddEnum8(FMemberId Name, FEnumSchemaId Schema, uint8 Value)	{ AddLeaf(Name, {ELeafType::Enum, ELeafWidth::B8},  ToOptional(Schema), Value); }
-	void AddEnum16(FMemberId Name, FEnumSchemaId Schema, uint16 Value)	{ AddLeaf(Name, {ELeafType::Enum, ELeafWidth::B16}, ToOptional(Schema), Value); }
-	void AddEnum32(FMemberId Name, FEnumSchemaId Schema, uint32 Value)	{ AddLeaf(Name, {ELeafType::Enum, ELeafWidth::B32}, ToOptional(Schema), Value); }
-	void AddEnum64(FMemberId Name, FEnumSchemaId Schema, uint64 Value)	{ AddLeaf(Name, {ELeafType::Enum, ELeafWidth::B64}, ToOptional(Schema), Value); }
+	template<UnsignedIntegral T>
+	void AddEnum(FMemberId Name, FEnumSchemaId Schema, T Value)
+	{
+		AddLeaf(Name, {ELeafType::Enum, ReflectArithmetic<T>.Width}, ToOptional(Schema), Value);
+	}
 	
-	void AddHex(FMemberId Name, uint8 Value)							{ AddLeaf(Name, {ELeafType::Hex, ELeafWidth::B8}, NoId, Value); }
-	void AddHex(FMemberId Name, uint16 Value)							{ AddLeaf(Name, {ELeafType::Hex, ELeafWidth::B16}, NoId, Value); }
-	void AddHex(FMemberId Name, uint32 Value)							{ AddLeaf(Name, {ELeafType::Hex, ELeafWidth::B32}, NoId, Value); }
-	void AddHex(FMemberId Name, uint64 Value)							{ AddLeaf(Name, {ELeafType::Hex, ELeafWidth::B64}, NoId, Value); }
+	template<UnsignedIntegral T>
+	void AddHex(FMemberId Name, T Value)
+	{
+		AddLeaf(Name, {ELeafType::Hex, ReflectArithmetic<T>.Width}, NoId, Value);
+	}
 
 	void AddLeaf(FMemberId Name, FUnpackedLeafType Leaf, FOptionalEnumSchemaId Enum, uint64 Value)	{ Members.Emplace(Name, Leaf, Enum, Value); }
 	void AddRange(FMemberId Name, FTypedRange Range)												{ Members.Emplace(Name, Range); }
-	void AddStruct(FMemberId Name, FStructSchemaId Schema, FBuiltStruct* Struct)					{ Members.Emplace(Name, Schema, Struct); }
+	// Add already built nested struct, must not be null
+	void AddStruct(FMemberId Name, FStructSchemaId Schema, FBuiltStruct* Struct)					{ check(Struct); Members.Emplace(Name, Schema, Struct); }
 	void Add(FMemberId Name, FTypedValue TypedValue)												{ Members.Emplace(Name, TypedValue); }
+
+	// Add already built super struct, must not be null and must be the first member added
+	PLAINPROPS_API void AddSuperStruct(FStructSchemaId SuperSchema, FBuiltStruct* SuperStruct);
 	
 	// Build members into a single nested super struct member, no-op if no non-super members has been added
 	PLAINPROPS_API void BuildSuperStruct(FScratchAllocator&	Scratch, const FStructDeclaration& Super, const FDebugIds& Debug);
diff --git a/Engine/Plugins/Experimental/PlainProps/Source/Public/PlainPropsParse.h b/Engine/Plugins/Experimental/PlainProps/Source/Public/PlainPropsParse.h
new file mode 100644
--- /dev/null
+++ b/Engine/Plugins/Experimental/PlainProps/Source/Public/PlainPropsParse.h
@@ -0,0 +1,38 @@
+// Copyright Epic Games, Inc. All Rights Reserved.
+
+#pragma once
+
+#include "PlainPropsTypes.h"
+#include "Containers/StringFwd.h"
+
+namespace PlainProps
+{
+
+template<typename T>	PLAINPROPS_API TOptional<T>						Parse(FUtf8StringView String);
+
+template<>				PLAINPROPS_API TOptional<ERangeSizeType>		Parse(FUtf8StringView String);
+template<>				PLAINPROPS_API TOptional<ELeafWidth>			Parse(FUtf8StringView String);
+
+template<>				PLAINPROPS_API TOptional<FUnpackedLeafType>		Parse(FUtf8StringView String);
+
+template<>				PLAINPROPS_API TOptional<bool>					Parse(FUtf8StringView String);
+template<>				PLAINPROPS_API TOptional<int8> 					Parse(FUtf8StringView String);
+template<>				PLAINPROPS_API TOptional<int16>					Parse(FUtf8StringView String);
+template<>				PLAINPROPS_API TOptional<int32>					Parse(FUtf8StringView String);
+template<>				PLAINPROPS_API TOptional<int64>					Parse(FUtf8StringView String);
+template<>				PLAINPROPS_API TOptional<uint8> 				Parse(FUtf8StringView String);
+template<>				PLAINPROPS_API TOptional<uint16>				Parse(FUtf8StringView String);
+template<>				PLAINPROPS_API TOptional<uint32>				Parse(FUtf8StringView String);
+template<>				PLAINPROPS_API TOptional<uint64>				Parse(FUtf8StringView String);
+template<>				PLAINPROPS_API TOptional<float>					Parse(FUtf8StringView String);
+template<>				PLAINPROPS_API TOptional<double>				Parse(FUtf8StringView String);
+template<>				PLAINPROPS_API TOptional<char>					Parse(FUtf8StringView String);
+template<>				PLAINPROPS_API TOptional<char8_t>				Parse(FUtf8StringView String);
+template<>				PLAINPROPS_API TOptional<char16_t>				Parse(FUtf8StringView String);
+template<>				PLAINPROPS_API TOptional<char32_t>				Parse(FUtf8StringView String);
+
+///////////////////////////////////////////////////////////////////////////////
+
+void PLAINPROPS_API ParseYamlBatch(TArray64<uint8>& OutBinary, FUtf8StringView Yaml);
+
+} // namespace PlainProps
diff --git a/Engine/Plugins/Experimental/PlainProps/Source/Public/PlainPropsPrint.h b/Engine/Plugins/Experimental/PlainProps/Source/Public/PlainPropsPrint.h
--- a/Engine/Plugins/Experimental/PlainProps/Source/Public/PlainPropsPrint.h
+++ b/Engine/Plugins/Experimental/PlainProps/Source/Public/PlainPropsPrint.h
@@ -3,7 +3,7 @@
 #pragma once
 
 #include "PlainPropsTypes.h"
-#include "Containers/StringView.h"
+#include "Containers/StringFwd.h"
 
 namespace PlainProps
 {
@@ -16,6 +16,7 @@ struct FDiffPath;
 
 PLAINPROPS_API FAnsiStringView ToString(ERangeSizeType Width);
 PLAINPROPS_API FAnsiStringView ToString(FUnpackedLeafType Leaf);
+PLAINPROPS_API FAnsiStringView ToString(ELeafWidth Width);
 
 ///////////////////////////////////////////////////////////////////////////////
 
@@ -26,7 +27,7 @@ PLAINPROPS_API void PrintMemberSchema(
 	FOptionalSchemaId InnerSchema,
 	TConstArrayView<FMemberType> InnerRangeTypes);
 
-void PrintBatchAsYamlDoc(
+PLAINPROPS_API void PrintYamlBatch(
 	FUtf8Builder& Out,
 	const FDebugIds& Names,
 	FReadBatchId Batch,
@@ -34,6 +35,4 @@ void PrintBatchAsYamlDoc(
 
 PLAINPROPS_API void PrintDiff(FUtf8Builder& Out, const FDiffPath& Diff, const FDebugIds& Names);
 
-///////////////////////////////////////////////////////////////////////////////
-
 } // namespace PlainProps
diff --git a/Engine/Plugins/Experimental/PlainProps/Source/Public/PlainPropsTypes.h b/Engine/Plugins/Experimental/PlainProps/Source/Public/PlainPropsTypes.h
--- a/Engine/Plugins/Experimental/PlainProps/Source/Public/PlainPropsTypes.h
+++ b/Engine/Plugins/Experimental/PlainProps/Source/Public/PlainPropsTypes.h
@@ -331,6 +331,10 @@ concept Arithmetic = std::is_arithmetic_v<T>;
 template<class T>
 concept Enumeration = std::is_enum_v<T>;
 
+// std::unsigned_integral and other std concepts are not supported by all tool chains yet
+template<class T>
+concept UnsignedIntegral = std::is_integral_v<T> && !std::is_signed_v<T>;
+
 template<class T>
 concept LeafType = std::is_arithmetic_v<T> || std::is_enum_v<T>;
 
diff --git a/Engine/Plugins/Experimental/PlainProps/Source/Tests/PlainPropsTests.cpp b/Engine/Plugins/Experimental/PlainProps/Source/Tests/PlainPropsTests.cpp
--- a/Engine/Plugins/Experimental/PlainProps/Source/Tests/PlainPropsTests.cpp
+++ b/Engine/Plugins/Experimental/PlainProps/Source/Tests/PlainPropsTests.cpp
@@ -15,8 +15,10 @@
 #include "PlainPropsIndex.h"
 #include "PlainPropsInternalBuild.h"
 #include "PlainPropsInternalFormat.h"
+#include "PlainPropsInternalParse.h"
 #include "PlainPropsInternalPrint.h"
 #include "PlainPropsInternalRead.h"
+#include "PlainPropsInternalTest.h"
 #include "Templates/UnrealTemplate.h"
 
 DEFINE_LOG_CATEGORY_STATIC(LogPlainPropsTests, Log, All);
@@ -31,13 +33,148 @@ static bool operator==(FParametricTypeView A, FParametricTypeView B)
 	return A.Name == B.Name && A.NumParameters == B.NumParameters && !FMemory::Memcmp(A.Parameters, B.Parameters, A.NumParameters * sizeof(FTypeId));
 }
 
-template<typename IntType>
-static TConstArrayView<IntType> GetEnumConstants(const FEnumSchema& Schema)
+static uint64 CalculateSize(const FSchemaBatch& In)
 {
-	CHECK(Schema.Width == ReflectLeaf<IntType>.Width);
-	return MakeConstArrayView<IntType>(reinterpret_cast<const IntType*>(Schema.Footer + Schema.Num), Schema.ExplicitConstants ? Schema.Num : 0);
+	uint32 NumParameters = 0;
+	for (FParametricType ParametricType : In.GetParametricTypes())
+	{
+		NumParameters += ParametricType.Parameters.NumParameters;
+	}
+	const uint8* End = reinterpret_cast<const uint8*>(In.GetFirstParameter() + NumParameters);
+	const uint64 Size = End - reinterpret_cast<const uint8*>(&In);
+	return Size;
+}
+
+static bool operator==(const FSchemaBatch& A, const FSchemaBatch& B)
+{
+	const uint64 SizeA = CalculateSize(A);
+	const uint64 SizeB = CalculateSize(B);
+	return SizeA == SizeB && FMemory::Memcmp(&A, &B, SizeA) == 0;
+}
+
+static bool operator==(const FStructSchema& A, const FStructSchema& B)
+{
+	const uint64 SizeA = CalculateSize(A);
+	const uint64 SizeB = CalculateSize(B);
+	return SizeA == SizeB && FMemory::Memcmp(&A, &B, SizeA) == 0;
+}
+
+static bool operator==(const FEnumSchema& A, const FEnumSchema& B)
+{
+	const uint64 SizeA = CalculateSize(A);
+	const uint64 SizeB = CalculateSize(B);
+	return SizeA == SizeB && FMemory::Memcmp(&A, &B, SizeA) == 0;
+}
+
+static bool operator!=(const FSchemaBatch& A, const FSchemaBatch& B)	{ return !(A == B); }
+static bool operator!=(const FStructSchema& A, const FStructSchema& B)	{ return !(A == B); }
+static bool operator!=(const FEnumSchema& A, const FEnumSchema& B)		{ return !(A == B); }
+
+////////////////////////////////////////////////////////////////////////////////
+
+static bool DiffStructSchema(const FStructSchema& A, const FStructSchema& B)
+{
+	const bool Diff = A != B;
+#if DO_CHECK
+	CHECK(A.Type == B.Type);
+	CHECK(A.NumMembers == B.NumMembers);
+	CHECK(A.NumRangeTypes == B.NumRangeTypes);
+	CHECK(A.NumInnerSchemas == B.NumInnerSchemas);
+	CHECK(A.Inheritance == B.Inheritance);
+	CHECK(A.IsDense == B.IsDense);
+	CHECK(EqualItems(A.GetMemberTypes(), B.GetMemberTypes()));
+	CHECK(EqualItems(A.GetRangeTypes(), B.GetRangeTypes()));
+	CHECK(EqualItems(A.GetMemberNames(), B.GetMemberNames()));
+	CHECK(EqualItems(
+		MakeArrayView(A.GetInnerSchemas(), A.NumInnerSchemas),
+		MakeArrayView(B.GetInnerSchemas(), B.NumInnerSchemas)));
+	CHECK(A.GetSuperSchema() == B.GetSuperSchema());
+	CHECK(A == B);
+#endif
+	return Diff;
+}
+
+static bool DiffEnumSchema(const FEnumSchema& A, const FEnumSchema& B)
+{
+	const bool Diff = A != B;
+#if DO_CHECK
+	CHECK(A.Type == B.Type);
+	CHECK(A.FlagMode == B.FlagMode);
+	CHECK(A.ExplicitConstants == B.ExplicitConstants);
+	CHECK(A.Width == B.Width);
+	CHECK(A.Num == B.Num);
+	CHECK(EqualItems(MakeArrayView(A.Footer, A.Num), MakeArrayView(B.Footer, B.Num)));
+	if (A.Width == B.Width)
+	{
+		switch (A.Width)
+		{
+		case ELeafWidth::B8:  CHECK(EqualItems(GetConstants<uint8 >(A), GetConstants<uint8 >(B))); break;
+		case ELeafWidth::B16: CHECK(EqualItems(GetConstants<uint16>(A), GetConstants<uint16>(B))); break;
+		case ELeafWidth::B32: CHECK(EqualItems(GetConstants<uint32>(A), GetConstants<uint32>(B))); break;
+		case ELeafWidth::B64: CHECK(EqualItems(GetConstants<uint64>(A), GetConstants<uint64>(B))); break;
+		}
+	}
+	CHECK(A == B);
+#endif
+	return Diff;
 }
 
+static bool DiffSchemas(const FSchemaBatch& A, const FSchemaBatch& B)
+{
+	const bool Diff = A != B;
+#if DO_CHECK
+	CHECK(A.NumNestedScopes == B.NumNestedScopes);
+	CHECK(A.NestedScopesOffset == B.NestedScopesOffset);
+	CHECK(A.NumParametricTypes == B.NumParametricTypes);
+	CHECK(A.NumSchemas == B.NumSchemas);
+	CHECK(A.NumStructSchemas == B.NumStructSchemas);
+	CHECK(EqualItems(A.GetSchemaOffsets(), B.GetSchemaOffsets()));
+	CHECK(EqualItems(A.GetNestedScopes(), B.GetNestedScopes()));
+	CHECK(EqualItems(A.GetParametricTypes(), B.GetParametricTypes()));
+
+	{
+		TSchemaRange<const FStructSchema> AA = GetStructSchemas(A);
+		TSchemaRange<const FStructSchema> BB = GetStructSchemas(B);
+		TSchemaIterator<const FStructSchema> ItA = AA.begin();
+		TSchemaIterator<const FStructSchema> ItB = BB.begin();
+		for (; ItA != AA.end() && ItB != BB.end(); ++ItA, ++ItB)
+		{
+			CHECK_FALSE(DiffStructSchema(*ItA, *ItB));
+		}
+	}
+
+	{
+		TSchemaRange<const FEnumSchema> AA = GetEnumSchemas(A);
+		TSchemaRange<const FEnumSchema> BB = GetEnumSchemas(B);
+		TSchemaIterator<const FEnumSchema> ItA = AA.begin();
+		TSchemaIterator<const FEnumSchema> ItB = BB.begin();
+		for (; ItA != AA.end() && ItB != BB.end(); ++ItA, ++ItB)
+		{
+			CHECK_FALSE(DiffEnumSchema(*ItA, *ItB));
+		}
+	}
+
+	CHECK(A == B);
+#endif
+	return Diff;
+}
+
+bool DiffSchemas(FMemoryView A, FMemoryView B)
+{
+	const bool BinaryDiff = A.GetSize() != B.GetSize() || FMemory::Memcmp(A.GetData(), B.GetData(), A.GetSize());
+#if DO_CHECK
+	const FSchemaBatch* SchemasA = reinterpret_cast<const FSchemaBatch*>(A.GetData());
+	const FSchemaBatch* SchemasB = reinterpret_cast<const FSchemaBatch*>(B.GetData());
+	SchemasA->ValidateBounds(A.GetSize());
+	SchemasB->ValidateBounds(B.GetSize());
+	const bool Diff = DiffSchemas(*SchemasA, *SchemasB);
+	checkf(Diff == BinaryDiff, TEXT("DiffSchemas() should produce the same result as a binary diff"));
+#endif
+	return BinaryDiff;
+}
+
+////////////////////////////////////////////////////////////////////////////////
+
 TEST_CASE_NAMED(FPlainPropsIndexTest, "System::Core::Serialization::PlainProps::Index", "[Core][PlainProps][SmokeFilter]")
 {
 	SECTION("NestedScope")
@@ -324,19 +461,27 @@ private:
 };
 
 static void PrintBatch(
+	FUtf8StringBuilderBase& Out,
 	const FTestNameReader& Names,
 	FReadBatchId Batch,
 	const FSchemaBatch& Schemas,
 	TConstArrayView<FStructView> Objects)
 {
 	FTestBatchNames BatchNames(Names, Schemas);
-	TUtf8StringBuilder<4096> YamlString;
-	PrintBatchAsYamlDoc(YamlString, BatchNames, Batch, Objects);
-	FUtf8StringView YamlView = YamlString.ToView();
+	PrintYamlBatch(Out, BatchNames, Batch, Objects);
+	FUtf8StringView YamlView = Out.ToView();
 	auto Wide = StringCast<TCHAR>(YamlView.GetData(), YamlView.Len());
 	UE_LOG(LogPlainPropsTests, Log, TEXT("\n%.*s\n"), Wide.Length(), Wide.Get());
 }
 
+static void ParseBatch(FUtf8StringView YamlView, FMemoryView BinaryView)
+{
+	TArray64<uint8> Data;
+	ParseYamlBatch(Data, YamlView);
+	FMemoryView ParsedBinaryView(Data.GetData(), Data.Num());
+	CHECK_FALSE(DiffSchemas(BinaryView, ParsedBinaryView));
+}
+
 ////////////////////////////////////////////////////////////////////////////////////////////////
 
 class FTestBatchReader
@@ -353,7 +498,8 @@ public:
 		CHECK(It.Grab<uint32>() == TestMagics[1]);
 		It.SkipAlignmentPadding<uint32>();
 		uint32 SchemasSize = It.Grab<uint32>();
-		const FSchemaBatch* Schemas = ValidateSchemas(It.GrabSlice(SchemasSize));
+		FMemoryView SchemasView = It.GrabSlice(SchemasSize);
+		const FSchemaBatch* Schemas = ValidateSchemas(SchemasView);
 		FReadBatchId Batch = MountReadSchemas(Schemas);
 		CHECK(It.Grab<uint32>() == TestMagics[2]);
 		
@@ -369,8 +515,10 @@ public:
 		CHECK(It.Grab<uint32>() == TestMagics[4]);
 		CHECK(!Objects.IsEmpty());
 
-		// Print schemas and objects
-		PrintBatch(Names, Batch, *Schemas, Objects);
+		// Roundtrip schemas and objects to text
+		TUtf8StringBuilder<4096> YamlString;
+		PrintBatch(YamlString, Names, Batch, *Schemas, Objects);
+		ParseBatch(YamlString.ToView(), SchemasView);
 	}
 
 	~FTestBatchReader()
@@ -776,19 +924,19 @@ TEST_CASE_NAMED(FPlainPropsReadWriteTest, "System::Core::Serialization::PlainPro
 			FEnumSchemaId I	= Batch.DeclareEnum("Test", "FlagLimit64",	EEnumMode::Flag, ELeafWidth::B64,	{"One", "Max"}, {1, 0x8000000000000000});
 					
 			FMemberBuilder Members;
-			Members.AddEnum8(Batch.NameMember("A2"),		A, 2);
-			Members.AddEnum8(Batch.NameMember("A0"),		A, 0);
-			Members.AddEnum8(Batch.NameMember("B0"),		B, 0);
-			Members.AddEnum8(Batch.NameMember("B4"),		B, 4);
-			Members.AddEnum8(Batch.NameMember("B5"),		B, 5);
-			Members.AddEnum8(Batch.NameMember("B7"),		B, 7);
-			Members.AddEnum(Batch.NameMember("C3"),			C, EFlatSparse8::C);
-			Members.AddEnum8(Batch.NameMember("D34"),		D, 34);
-			Members.AddEnum8(Batch.NameMember("Max8"),		E, 0xFF);
-			Members.AddEnum16(Batch.NameMember("Max16"),	F, 0xFFFF);
-			Members.AddEnum32(Batch.NameMember("Max32"),	G, 0xFFFFFFFF);
-			Members.AddEnum64(Batch.NameMember("Max64"),	H, 0xFFFFFFFFFFFFFFFF);
-			Members.AddEnum64(Batch.NameMember("IF"),		I, 0x8000000000000001);
+			Members.AddEnum(Batch.NameMember("A2"),				A, (uint8)2);
+			Members.AddEnum(Batch.NameMember("A0"),				A, (uint8)0);
+			Members.AddEnum(Batch.NameMember("B0"),				B, (uint8)0);
+			Members.AddEnum(Batch.NameMember("B4"),				B, (uint8)4);
+			Members.AddEnum(Batch.NameMember("B5"),				B, (uint8)5);
+			Members.AddEnum(Batch.NameMember("B7"),				B, (uint8)7);
+			Members.AddEnum(Batch.NameMember("C3"),				C, EFlatSparse8::C);
+			Members.AddEnum(Batch.NameMember("D34"),			D, (uint8)34);
+			Members.AddEnum(Batch.NameMember("Max8"),			E, (uint8)0xFF);
+			Members.AddEnum(Batch.NameMember("Max16"),			F, (uint16)0xFFFF);
+			Members.AddEnum(Batch.NameMember("Max32"),			G, (uint32)0xFFFFFFFF);
+			Members.AddEnum(Batch.NameMember("Max64"),			H, (uint64)0xFFFFFFFFFFFFFFFF);
+			Members.AddEnum(Batch.NameMember("IF"),				I, (uint64)0x8000000000000001);
 
 			Batch.AddObject(ObjectId, MoveTemp(Members));
 		}, 
@@ -869,16 +1017,16 @@ TEST_CASE_NAMED(FPlainPropsReadWriteTest, "System::Core::Serialization::PlainPro
 			CHECK(Names[FlagLimit64.Footer[0]] == "One");
 			CHECK(Names[FlagLimit64.Footer[1]] == "Max");
 
-			CHECK(EqualItems(GetEnumConstants<uint8>(FlatDense8), MakeConstArrayView({0, 2})));
-			CHECK(EqualItems(GetEnumConstants<uint8>(FlagDense8), TConstArrayView<uint8>()));
-			CHECK(EqualItems(GetEnumConstants<EFlatSparse8>(FlatSparse8), MakeConstArrayView({EFlatSparse8::C})));
-			CHECK(EqualItems(GetEnumConstants<uint8>(FlagSparse8), MakeConstArrayView({2, 32})));
-			CHECK(GetEnumConstants<uint8>(FlatLimit8)[0] == 0xFF);
-			CHECK(GetEnumConstants<uint16>(FlatLimit16)[0] == 0xFFFF);
-			CHECK(GetEnumConstants<uint32>(FlatLimit32)[0] == 0xFFFFFFFF);
-			CHECK(GetEnumConstants<uint64>(FlatLimit64)[0] == 0xFFFFFFFFFFFFFFFF);
-			CHECK(GetEnumConstants<uint64>(FlagLimit64)[0] == 1);
-			CHECK(GetEnumConstants<uint64>(FlagLimit64)[1] == 0x8000000000000000);
+			CHECK(EqualItems(GetConstants<uint8>(FlatDense8), MakeConstArrayView({0, 2})));
+			CHECK(EqualItems(GetConstants<uint8>(FlagDense8), TConstArrayView<uint8>()));
+			CHECK(EqualItems(GetConstants<EFlatSparse8>(FlatSparse8), MakeConstArrayView({EFlatSparse8::C})));
+			CHECK(EqualItems(GetConstants<uint8>(FlagSparse8), MakeConstArrayView({2, 32})));
+			CHECK(GetConstants<uint8>(FlatLimit8)[0] == 0xFF);
+			CHECK(GetConstants<uint16>(FlatLimit16)[0] == 0xFFFF);
+			CHECK(GetConstants<uint32>(FlatLimit32)[0] == 0xFFFFFFFF);
+			CHECK(GetConstants<uint64>(FlatLimit64)[0] == 0xFFFFFFFFFFFFFFFF);
+			CHECK(GetConstants<uint64>(FlagLimit64)[0] == 1);
+			CHECK(GetConstants<uint64>(FlagLimit64)[1] == 0x8000000000000000);
 		});
 	}
 
@@ -977,7 +1125,7 @@ TEST_CASE_NAMED(FPlainPropsReadWriteTest, "System::Core::Serialization::PlainPro
 
 		TestSerialize([](FTestBatchBuilder& Batch, FScratchAllocator& Scratch)
 		{
-			FStructSchemaId Object = Batch.DeclareStruct("Test", "Object", {"IntRs", "EmptyRs", "FloatRs", "EnumRs", "UnusedEnumRs", "StructRs", "StructRRs"}, EMemberPresence::AllowSparse);
+			FStructSchemaId Object = Batch.DeclareStruct("Test", "Object", {"IntRs", "EmptyRs", "FloatRs", "EnumRs", "UnusedEnumRs", "UnicodeRs", "StructRs", "StructRRs"}, EMemberPresence::AllowSparse);
 			FStructSchemaId XY = Batch.DeclareStruct("Test", "XY", {"X", "Y"}, EMemberPresence::RequireAll);
 			FStructSchemaId ZW = Batch.DeclareStruct("Test", "ZW", {"Z", "W"}, EMemberPresence::AllowSparse);
 			FEnumSchemaId Enum = Batch.DeclareEnum("Test", "AB", EEnumMode::Flag, ELeafWidth::B8, {"A", "B"}, {1, 4});
@@ -1001,6 +1149,11 @@ TEST_CASE_NAMED(FPlainPropsReadWriteTest, "System::Core::Serialization::PlainPro
 			UnusedEnumRs.Add({});
 			UnusedEnumRs.Add(BuildEnumRange(Scratch, UnusedEnum, TConstArrayView<EUnused>()));
 
+			FNestedRangeBuilder UnicodeRs(MakeLeafRangeSchema<char8_t, int32>(), 3);
+			UnicodeRs.Add(BuildLeafRange(Scratch, MakeArrayView(u8"Hello")));
+			UnicodeRs.Add({});
+			UnicodeRs.Add(BuildLeafRange(Scratch, MakeArrayView(u8"World!")));
+
 			FStructRangeBuilder XYs(uint64(2));
 			XYs[0].Add(Batch.NameMember("X"), 1.f);
 			XYs[0].Add(Batch.NameMember("Y"), 2.f);
@@ -1019,13 +1172,13 @@ TEST_CASE_NAMED(FPlainPropsReadWriteTest, "System::Core::Serialization::PlainPro
 			FNestedRangeBuilder StructRRs(MakeNestedRangeSchema(Scratch, ERangeSizeType::U32, ZWRangeSchema), 1);
 			StructRRs.Add(ZWRs.BuildAndReset(Scratch, ERangeSizeType::U32));
 
-
 			FMemberBuilder Members;
 			Members.AddRange(Batch.NameMember("IntRs"), IntRs.BuildAndReset(Scratch, ERangeSizeType::S32));
 			Members.AddRange(Batch.NameMember("EmptyRs"), IntRs.BuildAndReset(Scratch, ERangeSizeType::S32));
 			Members.AddRange(Batch.NameMember("FloatRs"), FloatRs.BuildAndReset(Scratch, ERangeSizeType::S64));
 			Members.AddRange(Batch.NameMember("EnumRs"), EnumRs.BuildAndReset(Scratch, ERangeSizeType::U8));
 			Members.AddRange(Batch.NameMember("UnusedEnumRs"), UnusedEnumRs.BuildAndReset(Scratch, ERangeSizeType::U8));
+			Members.AddRange(Batch.NameMember("UnicodeRs"), UnicodeRs.BuildAndReset(Scratch, ERangeSizeType::S32));
 			Members.AddRange(Batch.NameMember("StructRs"), StructRs.BuildAndReset(Scratch, ERangeSizeType::U64));
 			Members.AddRange(Batch.NameMember("StructRRs"), StructRRs.BuildAndReset(Scratch, ERangeSizeType::U32));
 
@@ -1041,6 +1194,7 @@ TEST_CASE_NAMED(FPlainPropsReadWriteTest, "System::Core::Serialization::PlainPro
 			TArray<FRangeView> FloatRs = MakeArray<FRangeView>(It.GrabRange().AsRanges());
 			TArray<FRangeView> EnumRs = MakeArray<FRangeView>(It.GrabRange().AsRanges());
 			TArray<FRangeView> UnusedEnumRs = MakeArray<FRangeView>(It.GrabRange().AsRanges());
+			TArray<FRangeView> UnicodeRs = MakeArray<FRangeView>(It.GrabRange().AsRanges());
 			TArray<FRangeView> StructRs = MakeArray<FRangeView>(It.GrabRange().AsRanges());
 			TArray<FRangeView> StructRRs = MakeArray<FRangeView>(It.GrabRange().AsRanges());
 
@@ -1350,7 +1504,7 @@ TEST_CASE_NAMED(FPlainPropsReadWriteTest, "System::Core::Serialization::PlainPro
 			FStructSchemaId B3 = Batch.DeclareStruct("Test", "B3", {"3"}, EMemberPresence::AllowSparse, ToOptional(B2)); // - - 0
 			FStructSchemaId B4 = Batch.DeclareStruct("Test", "B4", {"4"}, EMemberPresence::AllowSparse, ToOptional(B3)); // 1 1 1
 			FStructSchemaId B5 = Batch.DeclareStruct("Test", "B5", {"5"}, EMemberPresence::AllowSparse, ToOptional(B4)); // 1 1 0
-			FStructSchemaId B6 = Batch.DeclareStruct("Test", "B6", {"6"}, EMemberPresence::AllowSparse, ToOptional(B5)); // 0 - 0
+			FStructSchemaId B6 = Batch.DeclareStruct("Test", "B6", {   }, EMemberPresence::AllowSparse, ToOptional(B5)); // 0 - 0
 			FStructSchemaId C5 = Batch.DeclareStruct("Test", "C5", {"5"}, EMemberPresence::AllowSparse, ToOptional(B4)); // - - -
 			
 			FMemberBuilder A;
diff --git a/Engine/Plugins/Experimental/PlainProps/Source/Tests/PlainPropsUeCoreTests.cpp b/Engine/Plugins/Experimental/PlainProps/Source/Tests/PlainPropsUeCoreTests.cpp
--- a/Engine/Plugins/Experimental/PlainProps/Source/Tests/PlainPropsUeCoreTests.cpp
+++ b/Engine/Plugins/Experimental/PlainProps/Source/Tests/PlainPropsUeCoreTests.cpp
@@ -7,8 +7,10 @@
 #include "PlainPropsIndex.h"
 #include "PlainPropsInternalBuild.h"
 #include "PlainPropsInternalFormat.h"
+#include "PlainPropsInternalParse.h"
 #include "PlainPropsInternalPrint.h"
 #include "PlainPropsInternalRead.h"
+#include "PlainPropsInternalTest.h"
 #include "PlainPropsLoad.h"
 #include "PlainPropsRead.h"
 #include "PlainPropsSave.h"
@@ -319,37 +321,62 @@ private:
 	const FSchemaBatch& Schemas;
 };
 
-void PrintSavedSchemas(const FSchemaBatchIds& BatchIds, const FSchemaBatch& SavedSchemas)
+static void PrintBatch(
+	FUtf8StringBuilderBase& Out,
+	const FSchemaBatchIds& Ids,
+	const FSchemaBatch& Schemas,
+	TArray<FStructView>&& Objects)
 {
-	FReadBatchId Batch = MountReadSchemas(&SavedSchemas);
-	TUtf8StringBuilder<4096> YamlString;
+	FReadBatchId Batch = MountReadSchemas(&Schemas);
+	for (FStructView& Struct : Objects)
 	{
-		FYamlBuilderPtr YamlBuilder = MakeYamlBuilder(YamlString);
-		FBatchPrinter Printer(*YamlBuilder, BatchIds);
-		Printer.PrintSchemas(Batch);
+		Struct.Schema.Batch = Batch;
 	}
-	FUtf8StringView YamlView = YamlString.ToView();
-	auto Wide = StringCast<TCHAR>(YamlView.GetData(), YamlView.Len());
-	UE_LOG(LogPlainPropsUeCoreTests, Log, TEXT("Schemas:\n%.*s"), Wide.Length(), Wide.Get());
+	PrintYamlBatch(Out, Ids, Batch, Objects);
 	UnmountReadSchemas(Batch);
 }
 
-void PrintSavedSchemas(TConstArrayView<FName> SavedNames, const FSchemaBatch& SavedSchemas)
+static void PrintBatch(
+	FUtf8StringBuilderBase& Out,
+	TConstArrayView<FName> Names,
+	const FSchemaBatch& Schemas,
+	TConstArrayView<FStructView> Objects)
 {
-	if (SavedNames.IsEmpty())
+	if (Names.IsEmpty())
 	{
-		FMemoryBatchIds BatchIds(SavedSchemas);
-		PrintSavedSchemas(BatchIds, SavedSchemas);
+		FMemoryBatchIds Ids(Schemas);
+		PrintBatch(Out, Ids, Schemas, TArray<FStructView>(Objects));
 	}
 	else
 	{
-		FStableBatchIds BatchIds(SavedSchemas);
-		BatchIds.Names = SavedNames;
-		PrintSavedSchemas(BatchIds, SavedSchemas);
+		FStableBatchIds Ids(Schemas);
+		Ids.Names = Names;
+		PrintBatch(Out, Ids, Schemas, TArray<FStructView>(Objects));
 	}
+
+	FUtf8StringView YamlView = Out.ToView();
+	auto Wide = StringCast<TCHAR>(YamlView.GetData(), YamlView.Len());
+	UE_LOG(LogPlainPropsUeCoreTests, Log, TEXT("%s Schemas:\n%.*s"),
+		Names.IsEmpty() ? TEXT("InMemory") : TEXT("StableNames"),
+		Wide.Length(), Wide.Get());
 }
 
-void PrintLoadStruct(const FIdBinding& Binding, FStructView StructView, int32 LoadIdx)
+static void ParseBatch(FUtf8StringView YamlView, FMemoryView BinaryView, ESchemaFormat Format)
+{
+	TArray64<uint8> Data;
+	if (Format == ESchemaFormat::StableNames)
+	{
+		ParseYamlBatch(Data, YamlView);
+		FMemoryView ParsedBinaryView(Data.GetData(), Data.Num());
+		CHECK_FALSE(DiffSchemas(BinaryView, ParsedBinaryView));
+	}
+	else
+	{
+		ParseYamlBatch(Data, YamlView);
+	}
+}
+
+static void PrintLoadStruct(const FIdBinding& Binding, FStructView StructView, int32 LoadIdx)
 {
 	TUtf8StringBuilder<4096> YamlString;
 	{
@@ -383,11 +410,10 @@ public:
 		CHECK(It.Grab<uint32>() == Magics[1]);
 		It.SkipAlignmentPadding<uint32>();
 		uint32 SchemasSize = It.Grab<uint32>();
-		const FSchemaBatch* SavedSchemas = ValidateSchemas(It.GrabSlice(SchemasSize));
+		FMemoryView SavedSchemasView = It.GrabSlice(SchemasSize);
+		const FSchemaBatch* SavedSchemas = ValidateSchemas(SavedSchemasView);
 		CHECK(It.Grab<uint32>() == Magics[2]);
 
-		// Mount, print and unmount saved schemas
-		PrintSavedSchemas(Ids, *SavedSchemas);
 
 		FReadBatchId Batch;
 		TOptional<FIdTranslator> RuntimeIds;
@@ -430,6 +456,11 @@ public:
 		// Finally create load plans
 		TConstArrayView<FStructSchemaId> LoadStructIds = RuntimeIds ? RuntimeIds->Translation.GetStructIds(SavedSchemas->NumStructSchemas) : InRuntimeIds;
 		Plans = CreateLoadPlans(Batch, GTypes, Customs, GSchemas, LoadStructIds, Format);
+
+		// Mount, print, unmount, and then parse the saved batch
+		TUtf8StringBuilder<4096> YamlString;
+		PrintBatch(YamlString, Ids, *SavedSchemas, Objects);
+		ParseBatch(YamlString.ToView(), SavedSchemasView, Format);
 	}
 
 	~FBatchLoader()
