PlainProps: Save all bit field bools with leaf width B8.

#rnx

diff --git a/Engine/Plugins/Experimental/PlainProps/Source/Private/PlainPropsSave.cpp b/Engine/Plugins/Experimental/PlainProps/Source/Private/PlainPropsSave.cpp
--- a/Engine/Plugins/Experimental/PlainProps/Source/Private/PlainPropsSave.cpp
+++ b/Engine/Plugins/Experimental/PlainProps/Source/Private/PlainPropsSave.cpp
@@ -284,7 +284,7 @@ FBuiltRange* SaveLeafRange(const void* Range, const ILeafRangeBinding& Binding,
 
 static void SaveMember(FMemberBuilder& Out, const void* Struct, FMemberId Name, const FSaveContext& Ctx, FLeafMemberBinding Member)
 {
-	FUnpackedLeafType Type = { ToLeafType(Member.Leaf.Type), Member.Leaf.Width };
+	FUnpackedLeafType Type = ToUnpackedLeafType(Member.Leaf);
 	Out.AddLeaf(Name, Type, Member.Enum, SaveLeaf(At(Struct, Member.Offset), Member.Leaf));
 }
 
@@ -482,4 +482,4 @@ FRangeSaverBase::FRangeSaverBase(FScratchAllocator& Scratch, uint64 Num, SIZE_T
 #endif
 {}
 
-} // namespace PlainProps
\ No newline at end of file
+} // namespace PlainProps
diff --git a/Engine/Plugins/Experimental/PlainProps/Source/Public/PlainPropsBind.h b/Engine/Plugins/Experimental/PlainProps/Source/Public/PlainPropsBind.h
--- a/Engine/Plugins/Experimental/PlainProps/Source/Public/PlainPropsBind.h
+++ b/Engine/Plugins/Experimental/PlainProps/Source/Public/PlainPropsBind.h
@@ -177,6 +177,16 @@ struct FUnpackedLeafBindType
 	}
 };
 
+inline static constexpr FUnpackedLeafType ToUnpackedLeafType(FUnpackedLeafBindType Leaf)
+{
+	if (Leaf.Type == ELeafBindType::BitfieldBool)
+	{
+		return { ELeafType::Bool, ELeafWidth::B8 };
+	}
+	
+	return { ToLeafType(Leaf.Type), Leaf.Width };
+}
+
 // @pre Type != ELeafBindType::BitfieldBool
 inline FUnpackedLeafType UnpackNonBitfield(FLeafBindType Packed)
 {
