PlainProps: Use FUtf8Builder typedef in PlainPropsUeCoreTests.cpp.

#rnx

diff --git a/Engine/Plugins/Experimental/PlainProps/Source/Tests/PlainPropsUeCoreTests.cpp b/Engine/Plugins/Experimental/PlainProps/Source/Tests/PlainPropsUeCoreTests.cpp
--- a/Engine/Plugins/Experimental/PlainProps/Source/Tests/PlainPropsUeCoreTests.cpp
+++ b/Engine/Plugins/Experimental/PlainProps/Source/Tests/PlainPropsUeCoreTests.cpp
@@ -269,8 +269,8 @@ class FSchemaBatchIds : public FDebugIds
 public:
 	FSchemaBatchIds(const FSchemaBatch& InSchemas) : Schemas(InSchemas) {}
 
-	virtual FTypeId						Resolve(FEnumSchemaId Id) const override final		{ return ResolveEnumSchema(Schemas, Id).Type; }
-	virtual FTypeId						Resolve(FStructSchemaId Id) const override final	{ return ResolveStructSchema(Schemas, Id).Type; }
+	virtual FTypeId				Resolve(FEnumSchemaId Id) const override final		{ return ResolveEnumSchema(Schemas, Id).Type; }
+	virtual FTypeId				Resolve(FStructSchemaId Id) const override final	{ return ResolveStructSchema(Schemas, Id).Type; }
 
 protected:
 	const FSchemaBatch& Schemas;
@@ -282,9 +282,9 @@ public:
 	using FSchemaBatchIds::FSchemaBatchIds;
 	using FSchemaBatchIds::Resolve;
 
-	virtual FNestedScope		Resolve(FNestedScopeId Id) const override { return GNames.Resolve(Id); }
-	virtual FParametricTypeView	Resolve(FParametricTypeId Id) const override { return GNames.Resolve(Id); }
-	virtual void				AppendString(FUtf8StringBuilderBase& Out, FNameId Name) const override { GNames.ResolveName(Name).AppendString(Out); }
+	virtual FNestedScope		Resolve(FNestedScopeId Id) const override			{ return GNames.Resolve(Id); }
+	virtual FParametricTypeView	Resolve(FParametricTypeId Id) const override		{ return GNames.Resolve(Id); }
+	virtual void				AppendString(FUtf8Builder& Out, FNameId Name) const override { GNames.ResolveName(Name).AppendString(Out); }
 };
 
 class FStableBatchIds final : public FSchemaBatchIds
@@ -293,9 +293,9 @@ public:
 	using FSchemaBatchIds::FSchemaBatchIds;
 	using FSchemaBatchIds::Resolve;
 
-	virtual FNestedScope		Resolve(FNestedScopeId Id) const override { return ResolveNestedScope(Schemas, Id); }
-	virtual FParametricTypeView	Resolve(FParametricTypeId Id) const override { return ResolveParametricType(Schemas, Id); }
-	virtual void				AppendString(FUtf8StringBuilderBase& Out, FNameId Name) const override { Names[Name.Idx].AppendString(Out); }
+	virtual FNestedScope		Resolve(FNestedScopeId Id) const override			{ return ResolveNestedScope(Schemas, Id); }
+	virtual FParametricTypeView	Resolve(FParametricTypeId Id) const override 		{ return ResolveParametricType(Schemas, Id); }
+	virtual void				AppendString(FUtf8Builder& Out, FNameId Name) const override { Names[Name.Idx].AppendString(Out); }
 
 	TConstArrayView<FName> Names;
 };
@@ -305,15 +305,15 @@ class FLoadBatchNames : public FDebugIds
 public:
 	FLoadBatchNames(const FIdBinding& InBinding, const FSchemaBatch& InSchemas) : DebugIds(GNames), Binding(InBinding), Schemas(InSchemas) {}
 
-	virtual FNestedScope				Resolve(FNestedScopeId Id) const override final		{ return DebugIds.Resolve(Binding.Remap(Id)); }
-	virtual FParametricTypeView			Resolve(FParametricTypeId Id) const override final	{ return DebugIds.Resolve(Binding.Remap(Id)); }
-	virtual FTypeId						Resolve(FEnumSchemaId Id) const override final		{ return ResolveEnumSchema(Schemas, Id).Type; }
-	virtual FTypeId						Resolve(FStructSchemaId Id) const override final	{ return ResolveStructSchema(Schemas, Id).Type; }
+	virtual FNestedScope		Resolve(FNestedScopeId Id) const override final		{ return DebugIds.Resolve(Binding.Remap(Id)); }
+	virtual FParametricTypeView	Resolve(FParametricTypeId Id) const override final	{ return DebugIds.Resolve(Binding.Remap(Id)); }
+	virtual FTypeId				Resolve(FEnumSchemaId Id) const override final		{ return ResolveEnumSchema(Schemas, Id).Type; }
+	virtual FTypeId				Resolve(FStructSchemaId Id) const override final	{ return ResolveStructSchema(Schemas, Id).Type; }
 
 	using FDebugIds::AppendString;
-	virtual void AppendString(FUtf8StringBuilderBase& Out, FNameId Name) const override final			{ DebugIds.AppendString(Out, Name); }
-	virtual void AppendString(FUtf8StringBuilderBase& Out, FTypenameId Typename) const override final	{ DebugIds.AppendString(Out, Typename); }
-	virtual void AppendString(FUtf8StringBuilderBase& Out, FScopeId Scope) const override final			{ DebugIds.AppendString(Out, Scope); }
+	virtual void AppendString(FUtf8Builder& Out, FNameId Name) const override final			{ DebugIds.AppendString(Out, Name); }
+	virtual void AppendString(FUtf8Builder& Out, FTypenameId Typename) const override final	{ DebugIds.AppendString(Out, Typename); }
+	virtual void AppendString(FUtf8Builder& Out, FScopeId Scope) const override final		{ DebugIds.AppendString(Out, Scope); }
 
 private:
 	const FDebugIds& DebugIds;
@@ -322,7 +322,7 @@ private:
 };
 
 static void PrintBatch(
-	FUtf8StringBuilderBase& Out,
+	FUtf8Builder& Out,
 	const FSchemaBatchIds& Ids,
 	const FSchemaBatch& Schemas,
 	TArray<FStructView>&& Objects)
@@ -337,7 +337,7 @@ static void PrintBatch(
 }
 
 static void PrintBatch(
-	FUtf8StringBuilderBase& Out,
+	FUtf8Builder& Out,
 	TConstArrayView<FName> Names,
 	const FSchemaBatch& Schemas,
 	TConstArrayView<FStructView> Objects)
