PlainProps: Split out FIds (for id resolving and text conversion) from FDebugIds (for debug printing only).

#rnx

diff --git a/Engine/Plugins/Experimental/PlainProps/Source/Private/PlainPropsBuildSchema.cpp b/Engine/Plugins/Experimental/PlainProps/Source/Private/PlainPropsBuildSchema.cpp
--- a/Engine/Plugins/Experimental/PlainProps/Source/Private/PlainPropsBuildSchema.cpp
+++ b/Engine/Plugins/Experimental/PlainProps/Source/Private/PlainPropsBuildSchema.cpp
@@ -13,10 +13,10 @@
 namespace PlainProps
 {
 
-static FString PrintMemberSchema(const FDebugIds& Debug, FMemberType Type, FOptionalSchemaId InnerSchema, TConstArrayView<FMemberType> InnerRangeTypes)
+static FString PrintMemberSchema(const FIds& Ids, FMemberType Type, FOptionalSchemaId InnerSchema, TConstArrayView<FMemberType> InnerRangeTypes)
 {
 	TUtf8StringBuilder<256> Utf8SchemaStr;
-	PrintMemberSchema(Utf8SchemaStr, Debug, Type, InnerSchema, InnerRangeTypes);
+	PrintMemberSchema(Utf8SchemaStr, Ids, Type, InnerSchema, InnerRangeTypes);
 	FString SchemaStr(FStringView(StringCast<TCHAR>(Utf8SchemaStr.GetData(), Utf8SchemaStr.Len())));
 
 	FMemberType InnermostType = InnerRangeTypes.Num() > 0 ? InnerRangeTypes.Last() : Type;
@@ -40,9 +40,9 @@ static FString PrintMemberSchema(const FDebugIds& Debug, FMemberType Type, FOpti
 	return SchemaStr;
 }
 
-static FString PrintMemberSchema(const FDebugIds& Debug, FMemberSchema Schema)
+static FString PrintMemberSchema(const FIds& Ids, FMemberSchema Schema)
 {
-	return PrintMemberSchema(Debug, Schema.Type, Schema.InnerSchema, Schema.GetInnerRangeTypes());
+	return PrintMemberSchema(Ids, Schema.Type, Schema.InnerSchema, Schema.GetInnerRangeTypes());
 }
 
 //////////////////////////////////////////////////////////////////////////
@@ -67,7 +67,7 @@ struct FStructSchemaBuilder
 struct FEnumSchemaBuilder
 {
 	const FEnumDeclaration&						Declaration;
-	const FDebugIds&							Debug;
+	FDebugIds									Debug;
 	FEnumSchemaId								Id;
 	TOptional<ELeafWidth>						NotedWidth;
 	TSet<uint64>								NotedConstants;
@@ -79,17 +79,18 @@ struct FEnumSchemaBuilder
 
 //////////////////////////////////////////////////////////////////////////
 
-FSchemasBuilder::FSchemasBuilder(const FDeclarations& Types, const IStructBindIds& InBindIds, FScratchAllocator& InScratch, ESchemaFormat InFormat)
-: FSchemasBuilder(Types.GetStructs(), Types.GetEnums(), InBindIds, Types.GetDebug(), InScratch, InFormat)
+FSchemasBuilder::FSchemasBuilder(const FDeclarations& Types, const FIds& InIds, const IStructBindIds& InBindIds, FScratchAllocator& InScratch, ESchemaFormat InFormat)
+: FSchemasBuilder(Types.GetStructs(), Types.GetEnums(), InIds, InBindIds, InScratch, InFormat)
 {}
 
-FSchemasBuilder::FSchemasBuilder(FStructDeclarations InStructs, FEnumDeclarations InEnums, const IStructBindIds& InBindIds, const FDebugIds& InDebug, FScratchAllocator& InScratch, ESchemaFormat InFormat)
+FSchemasBuilder::FSchemasBuilder(FStructDeclarations InStructs, FEnumDeclarations InEnums, const FIds& InIds, const IStructBindIds& InBindIds, FScratchAllocator& InScratch, ESchemaFormat InFormat)
 : DeclaredStructs(InStructs)
 , DeclaredEnums(InEnums)
+, Ids(InIds)
 , BindIds(InBindIds)
 , Format(InFormat)
 , Scratch(InScratch)
-, Debug(InDebug)
+, Debug(InIds)
 {}
 
 FSchemasBuilder::~FSchemasBuilder() {}
@@ -257,8 +258,8 @@ void FStructSchemaBuilder::NoteMembersRecursively(const FBuiltStruct& Struct)
 				checkf(NotedSchema == Member.Schema, TEXT("Member '%s' in '%s' first added as '%s' and later as '%s'."),
 					*AllSchemas.GetDebug().Print(Member.Name),
 					*AllSchemas.GetDebug().Print(Declaration.Type),
-					*PrintMemberSchema(AllSchemas.GetDebug(), NotedSchema),
-					*PrintMemberSchema(AllSchemas.GetDebug(), Member.Schema));
+					*PrintMemberSchema(AllSchemas.GetIds(), NotedSchema),
+					*PrintMemberSchema(AllSchemas.GetIds(), Member.Schema));
 			}
 		}
 		else
@@ -355,7 +356,7 @@ void FStructSchemaBuilder::NoteRangeRecursively(ERangeSizeType NumType, TConstAr
 
 FBuiltStructSchema FStructSchemaBuilder::Build() const
 {
-	FTypeId Type = Id == Declaration.Id ? Declaration.Type : AllSchemas.GetDebug().Resolve(Id);
+	FTypeId Type = Id == Declaration.Id ? Declaration.Type : AllSchemas.GetIds().Resolve(Id);
 	FBuiltStructSchema Out = { Type, Id, Declaration.Super, /* dense */ true };
 	
 	if (int32 Num = NotedMembers.CountSetBits())
diff --git a/Engine/Plugins/Experimental/PlainProps/Source/Private/PlainPropsInternalParse.h b/Engine/Plugins/Experimental/PlainProps/Source/Private/PlainPropsInternalParse.h
--- a/Engine/Plugins/Experimental/PlainProps/Source/Private/PlainPropsInternalParse.h
+++ b/Engine/Plugins/Experimental/PlainProps/Source/Private/PlainPropsInternalParse.h
@@ -137,7 +137,7 @@ class FBatchParser
 {
 public:
 	FBatchParser(FYamlTokenizer& InTokenizer, FScratchAllocator& InScratch)
-		: Tokenizer(InTokenizer), Scratch(InScratch), Types(Names)
+		: Tokenizer(InTokenizer), Scratch(InScratch), Types(FDebugIds(Names))
 	{ }
 
 	~FBatchParser() {};
diff --git a/Engine/Plugins/Experimental/PlainProps/Source/Private/PlainPropsInternalPrint.h b/Engine/Plugins/Experimental/PlainProps/Source/Private/PlainPropsInternalPrint.h
--- a/Engine/Plugins/Experimental/PlainProps/Source/Private/PlainPropsInternalPrint.h
+++ b/Engine/Plugins/Experimental/PlainProps/Source/Private/PlainPropsInternalPrint.h
@@ -30,7 +30,7 @@ using FYamlBuilderPtr = TUniquePtr<FYamlBuilder, FYamlBuilderDeleter>;
 class FBatchPrinter
 {
 public:
-	FBatchPrinter(FYamlBuilder& InTextBuilder, const FDebugIds& InNames);
+	FBatchPrinter(FYamlBuilder& InTextBuilder, const FIds& InIds);
 	~FBatchPrinter();
 
 	void PrintSchemas(FReadBatchId Batch);
@@ -45,7 +45,7 @@ private:
 	void PrintEnumConstants(TConstArrayView<FNameId> EnumNames, TConstArrayView<IntType> Constants, bool bFlagMode);
 
 	FYamlBuilder& TextBuilder;
-	const FDebugIds& Names;
+	const FIds& Ids;
 };
 
 ///////////////////////////////////////////////////////////////////////////////
diff --git a/Engine/Plugins/Experimental/PlainProps/Source/Private/PlainPropsParse.cpp b/Engine/Plugins/Experimental/PlainProps/Source/Private/PlainPropsParse.cpp
--- a/Engine/Plugins/Experimental/PlainProps/Source/Private/PlainPropsParse.cpp
+++ b/Engine/Plugins/Experimental/PlainProps/Source/Private/PlainPropsParse.cpp
@@ -543,6 +543,7 @@ public:
 	FBuiltStruct* ParseMembers(FStructSchemaId SchemaId);
 
 private:
+	FDebugIds GetDebug() const { return FDebugIds(Names); }
 	bool HasMore() const { return MemberTokens.HasMore(); }
 	bool PeekKind(EMemberKind& MemberKind) const;
 	bool ParseAll();
@@ -1468,7 +1469,7 @@ bool FMemberParser::ParseDynamicStructType(FOptionalStructSchemaId& Out)
 FBuiltStruct* FMemberParser::BuildAndReset(FStructSchemaId Schema)
 {
 	const FStructDeclaration& StructDeclaration = Types.Get(static_cast<FStructSchemaId>(Schema));
-	return MemberBuilder.BuildAndReset(Scratch, StructDeclaration, Names);
+	return MemberBuilder.BuildAndReset(Scratch, StructDeclaration, GetDebug());
 }
 
 template<Arithmetic T>
@@ -1878,7 +1879,7 @@ public:
 void FBatchParser::Write(TArray64<uint8>& Out)
 {
 	FParseBindings Bindings;
-	FSchemasBuilder SchemaBuilders(Types, Bindings, Scratch, ESchemaFormat::StableNames);
+	FSchemasBuilder SchemaBuilders(Types, Names, Bindings, Scratch, ESchemaFormat::StableNames);
 	for (const TPair<FStructSchemaId, FBuiltStruct*>& Object : Objects)
 	{
 		SchemaBuilders.NoteStructAndMembers(Object.Key, *Object.Value);
diff --git a/Engine/Plugins/Experimental/PlainProps/Source/Private/PlainPropsPrint.cpp b/Engine/Plugins/Experimental/PlainProps/Source/Private/PlainPropsPrint.cpp
--- a/Engine/Plugins/Experimental/PlainProps/Source/Private/PlainPropsPrint.cpp
+++ b/Engine/Plugins/Experimental/PlainProps/Source/Private/PlainPropsPrint.cpp
@@ -57,7 +57,7 @@ void FIdIndexerBase::InitParameterNames()
 
 FUtf8StringView PrintMemberSchemaImpl(
 	FUtf8Builder& Out,
-	const FDebugIds& Names,
+	const FIds& Ids,
 	FMemberType Type,
 	FOptionalSchemaId InnerSchema = NoId,
 	TConstArrayView<FMemberType> InnerRangeTypes = {})
@@ -70,12 +70,12 @@ FUtf8StringView PrintMemberSchemaImpl(
 		if (InnermostType.IsLeaf())
 		{
 			check(IsEnum(InnermostType));
-			Names.AppendString(Out, static_cast<FEnumSchemaId>(InnerSchema.Get()));
+			Ids.AppendString(Out, static_cast<FEnumSchemaId>(InnerSchema.Get()));
 		}
 		else
 		{
 			check(InnermostType.IsStruct());
-			Names.AppendString(Out, static_cast<FStructSchemaId>(InnerSchema.Get()));
+			Ids.AppendString(Out, static_cast<FStructSchemaId>(InnerSchema.Get()));
 		}
 	}
 	else if (InnermostType.IsLeaf())
@@ -106,12 +106,12 @@ FUtf8StringView PrintMemberSchemaImpl(
 
 void PrintMemberSchema(
 	FUtf8Builder& Out,
-	const FDebugIds& Names, 
+	const FIds& Ids, 
 	FMemberType Type,
 	FOptionalSchemaId InnerSchema,
 	TConstArrayView<FMemberType> InnerRangeTypes)
 {
-	PrintMemberSchemaImpl(Out, Names, Type, InnerSchema, InnerRangeTypes);
+	PrintMemberSchemaImpl(Out, Ids, Type, InnerSchema, InnerRangeTypes);
 }
 
 ///////////////////////////////////////////////////////////////////////////////
@@ -233,9 +233,9 @@ private:
 class FMemberPrinter
 {
 public:
-	FMemberPrinter(FYamlBuilder& InTextBuilder, const FDebugIds& InNames)
+	FMemberPrinter(FYamlBuilder& InTextBuilder, const FIds& InIds)
 	: TextBuilder(InTextBuilder)
-	, Names(InNames)
+	, Ids(InIds)
 	{}
 
 	void PrintMembers(const FStructView& StructView);
@@ -257,19 +257,19 @@ private:
 	void PrintUnicodeRangeAsLeaf(FOptionalMemberId Id, ERangeSizeType NumType, const FRangeView& RangeView);
 
 	FYamlBuilder& TextBuilder;
-	const FDebugIds& Names;
+	const FIds& Ids;
 };
 
 ///////////////////////////////////////////////////////////////////////////////
 
 void PrintYamlBatch(
 	FUtf8Builder& Out,
-	const FDebugIds& Names,
+	const FIds& Ids,
 	FReadBatchId Batch,
 	TConstArrayView<FStructView> Objects)
 {
 	FYamlBuilder YamlBuilder(Out);
-	FBatchPrinter Printer(YamlBuilder, Names);
+	FBatchPrinter Printer(YamlBuilder, Ids);
 	YamlBuilder.BeginDocument();
 	Printer.PrintSchemas(Batch);
 	Printer.PrintObjects(Objects);
@@ -558,9 +558,9 @@ void FYamlBuilder::AppendIdentifier(FUtf8StringView Id)
 ///////////////////////////////////////////////////////////////////////////////
 
 
-FBatchPrinter::FBatchPrinter(FYamlBuilder& InTextBuilder, const FDebugIds& InNames)
+FBatchPrinter::FBatchPrinter(FYamlBuilder& InTextBuilder, const FIds& InIds)
 : TextBuilder(InTextBuilder)
-, Names(InNames)
+, Ids(InIds)
 {
 }
 
@@ -594,39 +594,39 @@ void FBatchPrinter::PrintObjects(TConstArrayView<FStructView> Objects)
 	TextBuilder.BeginRange(GLiterals.Objects);
 	for (FStructView Object : Objects)
 	{
-		FMemberPrinter(TextBuilder, Names).PrintMembers(Object);
+		FMemberPrinter(TextBuilder, Ids).PrintMembers(Object);
 	}
 	TextBuilder.EndRange();
 }
 
 static FUtf8StringView PrintMemberSchema(
 	FUtf8Builder& Out,
-	const FDebugIds& Names,
+	const FIds& Ids,
 	const FMemberSchemaView& Schema)
 {
-	return PrintMemberSchemaImpl(Out, Names, Schema.Type, Schema.InnerSchema, Schema.InnerRangeTypes);
+	return PrintMemberSchemaImpl(Out, Ids, Schema.Type, Schema.InnerSchema, Schema.InnerRangeTypes);
 }
 
 template <int32 BufferSize>
 class FPrintId
 {
-	TUtf8StringBuilder<BufferSize>			Buffer;
+	TUtf8StringBuilder<BufferSize>		Buffer;
 public:
 	template <typename T>
-	FPrintId(const FDebugIds& Names, T Id)	{ Names.AppendString(Buffer, Id); }
-	FUtf8StringView operator*() const		{ return Buffer.ToView(); }
+	FPrintId(const FIds& Ids, T Id)		{ Ids.AppendString(Buffer, Id); }
+	FUtf8StringView operator*() const	{ return Buffer.ToView(); }
 };
 
 void FBatchPrinter::PrintStructSchema(const FStructSchema& Struct, FReadBatchId BatchId)
 {
 	FStructSchemaReader Reader(Struct, BatchId);
 
-	TextBuilder.BeginStruct(*FPrintId<128>(Names, Reader.GetTypeId()));
+	TextBuilder.BeginStruct(*FPrintId<128>(Ids, Reader.GetTypeId()));
 
 	if (Reader.HasSuper())
 	{
 		const FStructSchema& SuperSchema = Reader.GetSuper().Resolve();
-		TextBuilder.AddLeaf(GLiterals.DeclaredSuper, *FPrintId<128>(Names, SuperSchema.Type));
+		TextBuilder.AddLeaf(GLiterals.DeclaredSuper, *FPrintId<128>(Ids, SuperSchema.Type));
 	}
 
 	TextBuilder.BeginRange(GLiterals.Members);
@@ -634,7 +634,7 @@ void FBatchPrinter::PrintStructSchema(const FStructSchema& Struct, FReadBatchId
 	{
 		TUtf8StringBuilder<256> Buf;
 		FOptionalMemberId MemberId = Reader.PeekName();
-		TextBuilder.AddLeaf(*FPrintId<128>(Names, MemberId), PrintMemberSchema(Buf, Names, Reader.GrabMember()));
+		TextBuilder.AddLeaf(*FPrintId<128>(Ids, MemberId), PrintMemberSchema(Buf, Ids, Reader.GrabMember()));
 	}
 	TextBuilder.EndRange();
 
@@ -643,7 +643,7 @@ void FBatchPrinter::PrintStructSchema(const FStructSchema& Struct, FReadBatchId
 
 void FBatchPrinter::PrintEnumSchema(const FEnumSchema& Enum)
 {
-	TextBuilder.BeginStruct(*FPrintId<128>(Names, Enum.Type));
+	TextBuilder.BeginStruct(*FPrintId<128>(Ids, Enum.Type));
 	TextBuilder.AddLeaf(GLiterals.FlagMode, !!Enum.FlagMode);
 	TextBuilder.AddLeaf(GLiterals.Width, ToString(Enum.Width));
 
@@ -681,7 +681,7 @@ void FBatchPrinter::PrintEnumConstants(
 		check(EnumNames.Num() == Constants.Num());
 		for (uint16 Idx = 0; Idx < NamesNum; ++Idx)
 		{
-			TextBuilder.AddLeaf(*FPrintId<128>(Names, EnumNames[Idx]), (uint64)Constants[Idx]);
+			TextBuilder.AddLeaf(*FPrintId<128>(Ids, EnumNames[Idx]), (uint64)Constants[Idx]);
 		}
 	}
 	else if (bFlagMode)
@@ -689,7 +689,7 @@ void FBatchPrinter::PrintEnumConstants(
 		uint64 Value = 1;
 		for (uint16 Idx = 0; Idx < NamesNum; ++Idx)
 		{
-			TextBuilder.AddLeaf(*FPrintId<128>(Names, EnumNames[Idx]), Value);
+			TextBuilder.AddLeaf(*FPrintId<128>(Ids, EnumNames[Idx]), Value);
 			Value <<= 1;
 		}
 	}
@@ -697,7 +697,7 @@ void FBatchPrinter::PrintEnumConstants(
 	{
 		for (uint16 Idx = 0; Idx < NamesNum; ++Idx)
 		{
-			TextBuilder.AddLeaf(*FPrintId<128>(Names, EnumNames[Idx]), (uint64)Idx);
+			TextBuilder.AddLeaf(*FPrintId<128>(Ids, EnumNames[Idx]), (uint64)Idx);
 		}
 	}
 }
@@ -707,13 +707,13 @@ void FBatchPrinter::PrintEnumConstants(
 void FMemberPrinter::PrintMembers(const FStructView& StructView)
 {
 	const FStructSchema& Schema = StructView.Schema.Resolve();
-	TextBuilder.BeginStruct(*FPrintId<128>(Names, Schema.Type));
+	TextBuilder.BeginStruct(*FPrintId<128>(Ids, Schema.Type));
 	PrintMembersInternal({ EMemberKind::Struct }, StructView, Schema);
 }
 
 void FMemberPrinter::PrintLeaf(FMemberId Id, const FLeafView& LeafView)
 {
-	TextBuilder.AddLeafId(*FPrintId<128>(Names, Id));
+	TextBuilder.AddLeafId(*FPrintId<128>(Ids, Id));
 
 	switch (LeafView.Leaf.Type)
 	{
@@ -769,14 +769,14 @@ void FMemberPrinter::PrintLeaf(FMemberId Id, const FLeafView& LeafView)
 	if constexpr(PrintWithComments)
 	{
 		TUtf8StringBuilder<256> Buf;
-		PrintMemberSchemaImpl(Buf, Names, FMemberType(LeafView.Leaf.Pack()), static_cast<FOptionalSchemaId>(LeafView.Enum));
+		PrintMemberSchemaImpl(Buf, Ids, FMemberType(LeafView.Leaf.Pack()), static_cast<FOptionalSchemaId>(LeafView.Enum));
 		TextBuilder.AddComment(Buf);
 	}
 }
 
 void FMemberPrinter::PrintStruct(FOptionalMemberId MemberId, FStructType StructType, const FStructView& StructView)
 {
-	TextBuilder.BeginStruct(*FPrintId<128>(Names, MemberId));
+	TextBuilder.BeginStruct(*FPrintId<128>(Ids, MemberId));
 	PrintMembersInternal(StructType, StructView, StructView.Schema.Resolve());
 }
 
@@ -788,7 +788,7 @@ void FMemberPrinter::PrintRange(FMemberId Id, ERangeSizeType NumType, const FRan
 	}
 	else
 	{
-		TextBuilder.BeginRange(*FPrintId<128>(Names, Id));
+		TextBuilder.BeginRange(*FPrintId<128>(Ids, Id));
 		PrintRangeInternal(NumType, RangeView);
 	}
 }
@@ -882,13 +882,13 @@ void FMemberPrinter::PrintMembersInternal(FStructType StructType, const FStructV
 		if (HasMembers)
 		{
 			TUtf8StringBuilder<256> Buf;
-			TextBuilder.AddComment(PrintMemberSchemaImpl(Buf, Names, FMemberType(StructType), StructView.Schema.Id));
+			TextBuilder.AddComment(PrintMemberSchemaImpl(Buf, Ids, FMemberType(StructType), StructView.Schema.Id));
 		}
 	}
 
 	if (StructType.IsDynamic)
 	{
-		TextBuilder.AddLeaf(GLiterals.Dynamic, *FPrintId<128>(Names, Schema.Type));
+		TextBuilder.AddLeaf(GLiterals.Dynamic, *FPrintId<128>(Ids, Schema.Type));
 	}
 	while (It.HasMore())
 	{
@@ -914,7 +914,7 @@ void FMemberPrinter::PrintMembersInternal(FStructType StructType, const FStructV
 		if (!HasMembers)
 		{
 			TUtf8StringBuilder<256> Buf;
-			TextBuilder.AddComment(PrintMemberSchemaImpl(Buf, Names, FMemberType(StructType), StructView.Schema.Id));
+			TextBuilder.AddComment(PrintMemberSchemaImpl(Buf, Ids, FMemberType(StructType), StructView.Schema.Id));
 		}
 	}
 }
@@ -933,12 +933,12 @@ static TArray<FMemberType, TInlineAllocator<16>> GetInnerRangeTypes(const FRange
 
 static void AddRangeSchemaComment(
 	FYamlBuilder& TextBuilder,
-	const FDebugIds& Names,
+	const FIds& Ids,
 	FMemberType Type,
 	const FRangeSchema& Schema)
 {
 	TUtf8StringBuilder<256> Comment;
-	TextBuilder.AddComment(PrintMemberSchemaImpl(Comment, Names, Type, Schema.InnermostSchema, GetInnerRangeTypes(Schema)));
+	TextBuilder.AddComment(PrintMemberSchemaImpl(Comment, Ids, Type, Schema.InnermostSchema, GetInnerRangeTypes(Schema)));
 }
 
 void FMemberPrinter::PrintRangeInternal(ERangeSizeType NumType, const FRangeView& RangeView)
@@ -948,7 +948,7 @@ void FMemberPrinter::PrintRangeInternal(ERangeSizeType NumType, const FRangeView
 	{
 		if (RangeView.Num() > 0)
 		{
-			AddRangeSchemaComment(TextBuilder, Names, FMemberType(NumType), Schema);
+			AddRangeSchemaComment(TextBuilder, Ids, FMemberType(NumType), Schema);
 		}
 	}
 
@@ -970,7 +970,7 @@ void FMemberPrinter::PrintRangeInternal(ERangeSizeType NumType, const FRangeView
 	{
 		if (RangeView.Num() == 0)
 		{
-			AddRangeSchemaComment(TextBuilder, Names, FMemberType(NumType), Schema);
+			AddRangeSchemaComment(TextBuilder, Ids, FMemberType(NumType), Schema);
 		}
 	}
 }
@@ -1033,7 +1033,7 @@ void FMemberPrinter::PrintUnicodeRangeAsLeaf(FOptionalMemberId Id, ERangeSizeTyp
 
 	if (Id)
 	{
-		TextBuilder.AddLeafId(*FPrintId<128>(Names, Id));
+		TextBuilder.AddLeafId(*FPrintId<128>(Ids, Id));
 	}
 
 	const FLeafRangeView LeafRange = RangeView.AsLeaves();
@@ -1049,18 +1049,18 @@ void FMemberPrinter::PrintUnicodeRangeAsLeaf(FOptionalMemberId Id, ERangeSizeTyp
 
 	if constexpr(PrintWithComments)
 	{
-		AddRangeSchemaComment(TextBuilder, Names, FMemberType(NumType), RangeView.Schema);
+		AddRangeSchemaComment(TextBuilder, Ids, FMemberType(NumType), RangeView.Schema);
 	}
 }
 
 ///////////////////////////////////////////////////////////////////////////////
 
-void FDebugIds::AppendString(FUtf8Builder& Out, FMemberId Name) const
+void FIds::AppendString(FUtf8Builder& Out, FMemberId Name) const
 {
 	AppendString(Out, Name.Id);
 }
 
-void FDebugIds::AppendString(FUtf8Builder& Out, FOptionalMemberId Name) const
+void FIds::AppendString(FUtf8Builder& Out, FOptionalMemberId Name) const
 {
 	if (Name)
 	{
@@ -1072,7 +1072,7 @@ void FDebugIds::AppendString(FUtf8Builder& Out, FOptionalMemberId Name) const
 	}
 }
 
-void FDebugIds::AppendString(FUtf8Builder& Out, FScopeId Scope) const
+void FIds::AppendString(FUtf8Builder& Out, FScopeId Scope) const
 {
 	if (Scope.IsFlat())
 	{
@@ -1087,7 +1087,7 @@ void FDebugIds::AppendString(FUtf8Builder& Out, FScopeId Scope) const
 	}
 }
 
-void FDebugIds::AppendString(FUtf8Builder& Out, FTypenameId Typename) const
+void FIds::AppendString(FUtf8Builder& Out, FTypenameId Typename) const
 {
 	if (Typename.IsConcrete())
 	{
@@ -1117,7 +1117,7 @@ void FDebugIds::AppendString(FUtf8Builder& Out, FTypenameId Typename) const
 	}
 }
 
-void FDebugIds::AppendString(FUtf8Builder& Out, FTypeId Type) const
+void FIds::AppendString(FUtf8Builder& Out, FTypeId Type) const
 {
 	if (Type.Scope)
 	{
@@ -1127,48 +1127,50 @@ void FDebugIds::AppendString(FUtf8Builder& Out, FTypeId Type) const
 	AppendString(Out, Type.Name);
 }
 
-void FDebugIds::AppendString(FUtf8Builder& Out, FEnumSchemaId Name) const
+void FIds::AppendString(FUtf8Builder& Out, FEnumSchemaId Name) const
 {
 	AppendString(Out, Resolve(Name));
 }
 
-void FDebugIds::AppendString(FUtf8Builder& Out, FStructSchemaId Name) const
+void FIds::AppendString(FUtf8Builder& Out, FStructSchemaId Name) const
 {
 	AppendString(Out, Resolve(Name));
 }
 
+////////////////////////////////////////////////////////////////////////////////
+
 FString FDebugIds::Print(FNameId Name) const
 {
 	TUtf8StringBuilder<128> Out;
-	AppendString(Out, Name);
+	Ids.AppendString(Out, Name);
 	return FString(FStringView(StringCast<TCHAR>(Out.GetData(), Out.Len())));
 }
 
 FString FDebugIds::Print(FMemberId Name) const
 {
 	TUtf8StringBuilder<128> Out;
-	AppendString(Out, Name.Id);
+	Ids.AppendString(Out, Name.Id);
 	return FString(FStringView(StringCast<TCHAR>(Out.GetData(), Out.Len())));
 }
 
 FString FDebugIds::Print(FOptionalMemberId Name) const
 {
 	TUtf8StringBuilder<128> Out;
-	AppendString(Out, Name);
+	Ids.AppendString(Out, Name);
 	return FString(FStringView(StringCast<TCHAR>(Out.GetData(), Out.Len())));
 }
 
 FString FDebugIds::Print(FScopeId Scope) const
 {
 	TUtf8StringBuilder<128> Out;
-	AppendString(Out, Scope);
+	Ids.AppendString(Out, Scope);
 	return FString(FStringView(StringCast<TCHAR>(Out.GetData(), Out.Len())));
 }
 
 FString FDebugIds::Print(FTypenameId Typename) const
 {
 	TUtf8StringBuilder<128> Out;
-	AppendString(Out, Typename);
+	Ids.AppendString(Out, Typename);
 	return FString(FStringView(StringCast<TCHAR>(Out.GetData(), Out.Len())));
 }
 
@@ -1180,33 +1182,33 @@ FString FDebugIds::Print(FConcreteTypenameId Typename) const
 FString FDebugIds::Print(FTypeId Type) const
 {
 	TUtf8StringBuilder<128> Out;
-	AppendString(Out, Type);
+	Ids.AppendString(Out, Type);
 	return FString(FStringView(StringCast<TCHAR>(Out.GetData(), Out.Len())));
 }
 
 FString FDebugIds::Print(FEnumSchemaId Name) const
 {
 	TUtf8StringBuilder<128> Out;
-	AppendString(Out, Name);
+	Ids.AppendString(Out, Name);
 	return FString(FStringView(StringCast<TCHAR>(Out.GetData(), Out.Len())));
 }
 
 FString FDebugIds::Print(FStructSchemaId Name) const
 {
 	TUtf8StringBuilder<128> Out;
-	AppendString(Out, Name);
+	Ids.AppendString(Out, Name);
 	return FString(FStringView(StringCast<TCHAR>(Out.GetData(), Out.Len())));
 }
 
 ///////////////////////////////////////////////////////////////////////////////
 
-void PrintDiff(FUtf8Builder& Out, const FDiffPath& Diff, const FDebugIds& Names)
+void PrintDiff(FUtf8Builder& Out, const FDiffPath& Diff, const FIds& Ids)
 {
 	check(Diff.Num());
 
 	for (FDiffNode Node : ReverseIterate(Diff))
 	{
-		Names.AppendString(Out, Node.Name);
+		Ids.AppendString(Out, Node.Name);
 		Out << '.';
 	}
 	Out.RemoveSuffix(1);
@@ -1216,15 +1218,15 @@ void PrintDiff(FUtf8Builder& Out, const FDiffPath& Diff, const FDebugIds& Names)
 	{
 		if (Node.Type.IsStruct())
 		{
-			Names.AppendString(Out, Names.Resolve(Node.Meta.Struct).Name);		
+			Ids.AppendString(Out, Ids.Resolve(Node.Meta.Struct).Name);		
 		}
 		else if (Node.Type.IsRange())
 		{
-			Names.AppendString(Out, FTypenameId(Node.Meta.Range.GetBindName()));	
+			Ids.AppendString(Out, FTypenameId(Node.Meta.Range.GetBindName()));	
 		}
 		else if (FOptionalEnumSchemaId Enum = Node.Meta.Leaf)
 		{
-			Names.AppendString(Out, Names.Resolve(Enum.Get()).Name);
+			Ids.AppendString(Out, Ids.Resolve(Enum.Get()).Name);
 		}
 		else
 		{
diff --git a/Engine/Plugins/Experimental/PlainProps/Source/Private/PlainPropsWrite.cpp b/Engine/Plugins/Experimental/PlainProps/Source/Private/PlainPropsWrite.cpp
--- a/Engine/Plugins/Experimental/PlainProps/Source/Private/PlainPropsWrite.cpp
+++ b/Engine/Plugins/Experimental/PlainProps/Source/Private/PlainPropsWrite.cpp
@@ -37,7 +37,7 @@ struct FWriteIds
 	TArray<FTypeId>						KeptParameters;
 	TArray<FNameId>						KeptNames;
 
-	const FDebugIds&					Debug;
+	FDebugIds							Debug;
 
 	FNameId								Remap(FNameId Old) const				{ return Names[Old.Idx].Get(); }
 	FMemberId							Remap(FMemberId Old) const				{ return { Remap(Old.Id) }; }
diff --git a/Engine/Plugins/Experimental/PlainProps/Source/Public/PlainPropsBind.h b/Engine/Plugins/Experimental/PlainProps/Source/Public/PlainPropsBind.h
--- a/Engine/Plugins/Experimental/PlainProps/Source/Public/PlainPropsBind.h
+++ b/Engine/Plugins/Experimental/PlainProps/Source/Public/PlainPropsBind.h
@@ -313,7 +313,7 @@ class FCustomBindingsBottom final : public FCustomBindings
 	mutable TArray<FInnerStruct>				AllInners;
 	PLAINPROPS_API virtual FCustomBindingHandle	Find(FStructSchemaId BindId) const override;	
 public:
-	explicit FCustomBindingsBottom(const FDebugIds& Dbg) : FCustomBindings(/* unconstructed */  AllInners, Dbg) {}
+	explicit FCustomBindingsBottom(FDebugIds Dbg) : FCustomBindings(/* unconstructed */  AllInners, Dbg) {}
 };
 
 class FCustomBindingsOverlay final : public FCustomBindings
@@ -605,7 +605,7 @@ class FSchemaBindings : public IStructBindIds
 {
 public:
 	UE_NONCOPYABLE(FSchemaBindings);
-	explicit FSchemaBindings(const FDebugIds& In) : Debug(In) {}
+	explicit FSchemaBindings(FDebugIds In) : Debug(In) {}
 	PLAINPROPS_API ~FSchemaBindings();
 
 	PLAINPROPS_API void						BindStruct(FStructSchemaId BindId, FStructSchemaId DeclId, TConstArrayView<FMemberBinding> Schema);
@@ -616,7 +616,7 @@ public:
 	PLAINPROPS_API virtual FStructSchemaId	GetDeclId(FStructSchemaId BindId) const override final;
 private:
 	TArray<TUniquePtr<FSchemaBinding>>		Bindings;
-	const FDebugIds&						Debug;
+	FDebugIds								Debug;
 };
 
 ////////////////////////////////////////////////////////////////////////////////////////////////
diff --git a/Engine/Plugins/Experimental/PlainProps/Source/Public/PlainPropsBuildSchema.h b/Engine/Plugins/Experimental/PlainProps/Source/Public/PlainPropsBuildSchema.h
--- a/Engine/Plugins/Experimental/PlainProps/Source/Public/PlainPropsBuildSchema.h
+++ b/Engine/Plugins/Experimental/PlainProps/Source/Public/PlainPropsBuildSchema.h
@@ -47,8 +47,8 @@ public:
 	using FStructDeclarations = TConstArrayView<TUniquePtr<FStructDeclaration>>;
 	using FEnumDeclarations = TConstArrayView<TUniquePtr<FEnumDeclaration>>;
 
-	PLAINPROPS_API FSchemasBuilder(const FDeclarations& Declarations, const IStructBindIds& BindIds, FScratchAllocator& Scratch, ESchemaFormat InFormat);
-	PLAINPROPS_API FSchemasBuilder(FStructDeclarations InStructs, FEnumDeclarations InEnums, const IStructBindIds& BindIds, const FDebugIds& InDebug, FScratchAllocator& Scratch, ESchemaFormat InFormat);
+	PLAINPROPS_API FSchemasBuilder(const FDeclarations& Types, const FIds& InIds, const IStructBindIds& InBindIds, FScratchAllocator& InScratch, ESchemaFormat InFormat);
+	PLAINPROPS_API FSchemasBuilder(FStructDeclarations InStructs, FEnumDeclarations InEnums, const FIds& InIds, const IStructBindIds& BindIds, FScratchAllocator& Scratch, ESchemaFormat InFormat);
 	PLAINPROPS_API ~FSchemasBuilder();
 
 	PLAINPROPS_API FEnumSchemaBuilder&			NoteEnum(FEnumSchemaId Id);
@@ -57,19 +57,21 @@ public:
 	PLAINPROPS_API FBuiltSchemas				Build();
 	
 	FScratchAllocator&							GetScratch() const { return Scratch; }
-	const FDebugIds&							GetDebug() const { return Debug; }
+	const FIds&									GetIds() const { return Ids; }
+	FDebugIds									GetDebug() const { return Debug; }
 
 private:
 	FStructDeclarations							DeclaredStructs;
 	FEnumDeclarations							DeclaredEnums;
 	TSet<FStructSchemaId>						StructIndices;
 	TSet<FEnumSchemaId>							EnumIndices;
+	const FIds&									Ids;
 	const IStructBindIds&						BindIds;
 	ESchemaFormat								Format;
 	TPagedArray<FStructSchemaBuilder, 4096>		Structs;		// TPagedArray for stable references
 	TPagedArray<FEnumSchemaBuilder, 4096>		Enums;			// TPagedArray for stable references
 	FScratchAllocator&							Scratch;
-	const FDebugIds&							Debug;
+	FDebugIds									Debug;
 	bool										bBuilt = false;
 
 	void										NoteInheritanceChains();
@@ -78,4 +80,4 @@ private:
 // Extract runtime ids to side-channel for loading to avoid reindexing with IndexRuntimeIds()
 [[nodiscard]] PLAINPROPS_API TArray<FStructSchemaId> ExtractRuntimeIds(const FBuiltSchemas& In);
 
-} // namespace PlainProps
\ No newline at end of file
+} // namespace PlainProps
diff --git a/Engine/Plugins/Experimental/PlainProps/Source/Public/PlainPropsDeclare.h b/Engine/Plugins/Experimental/PlainProps/Source/Public/PlainPropsDeclare.h
--- a/Engine/Plugins/Experimental/PlainProps/Source/Public/PlainPropsDeclare.h
+++ b/Engine/Plugins/Experimental/PlainProps/Source/Public/PlainPropsDeclare.h
@@ -60,7 +60,7 @@ class FDeclarations
 {
 public:
 	UE_NONCOPYABLE(FDeclarations);
-	explicit FDeclarations(const FDebugIds& In) : Debug(In) {}
+	explicit FDeclarations(FDebugIds In) : Debug(In) {}
 
 	PLAINPROPS_API const FEnumDeclaration&			DeclareEnum(FEnumSchemaId DeclId, FTypeId Type, EEnumMode Mode, TConstArrayView<FEnumerator> Enumerators, EEnumAliases Policy);
 	// Declare struct with ref count 1 or increment it and check that previous declaration matches
@@ -76,12 +76,12 @@ public:
 	
 	TConstArrayView<TUniquePtr<FEnumDeclaration>>	GetEnums() const				{ return DeclaredEnums; }
 	TConstArrayView<TUniquePtr<FStructDeclaration>>	GetStructs() const				{ return DeclaredStructs; }
-	const FDebugIds&								GetDebug() const				{ return Debug; }
+	FDebugIds										GetDebug() const				{ return Debug; }
 
 protected:
 	TArray<TUniquePtr<FEnumDeclaration>>			DeclaredEnums;
 	TArray<TUniquePtr<FStructDeclaration>>			DeclaredStructs;
-	const FDebugIds&								Debug;
+	FDebugIds										Debug;
 
 #if DO_CHECK
 	PLAINPROPS_API void								Check(FEnumSchemaId Id) const;
@@ -97,4 +97,4 @@ struct IStructBindIds
 	virtual FStructSchemaId GetDeclId(FStructSchemaId BindId) const = 0;
 };
 
-} // namespace PlainProps
\ No newline at end of file
+} // namespace PlainProps
diff --git a/Engine/Plugins/Experimental/PlainProps/Source/Public/PlainPropsIndex.h b/Engine/Plugins/Experimental/PlainProps/Source/Public/PlainPropsIndex.h
--- a/Engine/Plugins/Experimental/PlainProps/Source/Public/PlainPropsIndex.h
+++ b/Engine/Plugins/Experimental/PlainProps/Source/Public/PlainPropsIndex.h
@@ -46,7 +46,7 @@ private:
 
 //////////////////////////////////////////////////////////////////////////
 
-class FIdIndexerBase : public FDebugIds
+class FIdIndexerBase : public FIds
 {
 public:
 	PLAINPROPS_API FScopeId				NestFlatScope(FScopeId Outer, FFlatScopeId Inner);
diff --git a/Engine/Plugins/Experimental/PlainProps/Source/Public/PlainPropsPrint.h b/Engine/Plugins/Experimental/PlainProps/Source/Public/PlainPropsPrint.h
--- a/Engine/Plugins/Experimental/PlainProps/Source/Public/PlainPropsPrint.h
+++ b/Engine/Plugins/Experimental/PlainProps/Source/Public/PlainPropsPrint.h
@@ -22,17 +22,17 @@ PLAINPROPS_API FAnsiStringView ToString(ELeafWidth Width);
 
 PLAINPROPS_API void PrintMemberSchema(
 	FUtf8Builder& Out,
-	const FDebugIds& Names, 
+	const FIds& Ids, 
 	FMemberType Type,
 	FOptionalSchemaId InnerSchema,
 	TConstArrayView<FMemberType> InnerRangeTypes);
 
 PLAINPROPS_API void PrintYamlBatch(
 	FUtf8Builder& Out,
-	const FDebugIds& Names,
+	const FIds& Ids,
 	FReadBatchId Batch,
 	TConstArrayView<FStructView> Objects);
 
-PLAINPROPS_API void PrintDiff(FUtf8Builder& Out, const FDiffPath& Diff, const FDebugIds& Names);
+PLAINPROPS_API void PrintDiff(FUtf8Builder& Out, const FDiffPath& Diff, const FIds& Ids);
 
 } // namespace PlainProps
diff --git a/Engine/Plugins/Experimental/PlainProps/Source/Public/PlainPropsTypes.h b/Engine/Plugins/Experimental/PlainProps/Source/Public/PlainPropsTypes.h
--- a/Engine/Plugins/Experimental/PlainProps/Source/Public/PlainPropsTypes.h
+++ b/Engine/Plugins/Experimental/PlainProps/Source/Public/PlainPropsTypes.h
@@ -420,10 +420,10 @@ enum class ESchemaFormat { StableNames, InMemoryNames };
 
 using FUtf8Builder = FUtf8StringBuilderBase;
 
-class FDebugIds
+class FIds
 {
 public:
-	virtual ~FDebugIds() {}
+	virtual ~FIds() {}
 
 	virtual FParametricTypeView			Resolve(FParametricTypeId Id) const = 0;
 	virtual FNestedScope				Resolve(FNestedScopeId Id) const = 0;
@@ -437,8 +437,14 @@ public:
 	PLAINPROPS_API virtual void			AppendString(FUtf8Builder& Out, FScopeId Scope) const;
 	PLAINPROPS_API virtual void			AppendString(FUtf8Builder& Out, FTypenameId Typename) const;
 	PLAINPROPS_API virtual void			AppendString(FUtf8Builder& Out, FTypeId Type) const;
-	PLAINPROPS_API virtual void			AppendString(FUtf8Builder& Out, FEnumSchemaId Name) const;
-	PLAINPROPS_API virtual void			AppendString(FUtf8Builder& Out, FStructSchemaId Name) const;
+	PLAINPROPS_API virtual void			AppendString(FUtf8Builder& Out, FEnumSchemaId Enum) const;
+	PLAINPROPS_API virtual void			AppendString(FUtf8Builder& Out, FStructSchemaId Struct) const;
+};
+
+class FDebugIds
+{
+public:
+	PLAINPROPS_API explicit	FDebugIds(const FIds& InIds) : Ids(InIds) {}
 
 	PLAINPROPS_API FString				Print(FNameId Name) const;
 	PLAINPROPS_API FString				Print(FMemberId Name) const;
@@ -447,8 +453,10 @@ public:
 	PLAINPROPS_API FString				Print(FTypenameId Typename) const;
 	PLAINPROPS_API FString				Print(FConcreteTypenameId Typename) const;
 	PLAINPROPS_API FString				Print(FTypeId Type) const;
-	PLAINPROPS_API FString				Print(FEnumSchemaId Name) const;
-	PLAINPROPS_API FString				Print(FStructSchemaId Name) const;
+	PLAINPROPS_API FString				Print(FEnumSchemaId Enum) const;
+	PLAINPROPS_API FString				Print(FStructSchemaId Struct) const;
+private:
+	const FIds& Ids;
 };
 
 } // namespace PlainProps
diff --git a/Engine/Plugins/Experimental/PlainProps/Source/Public/PlainPropsWrite.h b/Engine/Plugins/Experimental/PlainProps/Source/Public/PlainPropsWrite.h
--- a/Engine/Plugins/Experimental/PlainProps/Source/Public/PlainPropsWrite.h
+++ b/Engine/Plugins/Experimental/PlainProps/Source/Public/PlainPropsWrite.h
@@ -21,7 +21,7 @@ struct IStructBindIds;
 class FWriter
 {
 public:
-	PLAINPROPS_API FWriter(const FIdIndexerBase& DeclaredIds, const IStructBindIds& BindIds, const FBuiltSchemas& InSchemas, ESchemaFormat Format);
+	PLAINPROPS_API FWriter(const FIdIndexerBase& AllIds, const IStructBindIds& BindIds, const FBuiltSchemas& InSchemas, ESchemaFormat Format);
 	PLAINPROPS_API ~FWriter();
 	
 	PLAINPROPS_API TConstArrayView<FNameId>			GetUsedNames() const;
@@ -32,7 +32,7 @@ public:
 
 private:
 	const FBuiltSchemas&							Schemas;
-	const FDebugIds&								Debug;
+	FDebugIds										Debug;
 	TUniquePtr<FWriteIds>							NewIds;
 };
 
diff --git a/Engine/Plugins/Experimental/PlainProps/Source/Tests/PlainPropsTests.cpp b/Engine/Plugins/Experimental/PlainProps/Source/Tests/PlainPropsTests.cpp
--- a/Engine/Plugins/Experimental/PlainProps/Source/Tests/PlainPropsTests.cpp
+++ b/Engine/Plugins/Experimental/PlainProps/Source/Tests/PlainPropsTests.cpp
@@ -264,7 +264,7 @@ void AppendString(FUtf8StringBuilderBase& Out, const FAnsiString& Str)
 class FTestBatchBuilder : public TIdIndexer<FAnsiString>, public IStructBindIds
 {
 public:
-	FTestBatchBuilder(FScratchAllocator& InScratch) : Declarations(*this), Scratch(InScratch) {}
+	FTestBatchBuilder(FScratchAllocator& InScratch) : Declarations(FDebugIds(*this)), Scratch(InScratch) {}
 
 	FEnumSchemaId				DeclareEnum(FTypeId Type, EEnumMode Mode, ELeafWidth Width, std::initializer_list<const char*> Names, std::initializer_list<uint64> Constants);
 	FEnumSchemaId				DeclareEnum(const char* Scope, const char* Name, EEnumMode Mode, ELeafWidth Width, std::initializer_list<const char*> Names, std::initializer_list<uint64> Constants);
@@ -279,6 +279,8 @@ public:
 
 	TArray64<uint8>				Write();
 
+	FDebugIds					GetDebug() const { return FDebugIds(*this); }
+
 private:
 	TArray<TPair<FStructSchemaId, FBuiltStruct*>>	Objects;
 	FDeclarations									Declarations;
@@ -342,13 +344,13 @@ FEnumSchemaId FTestBatchBuilder::DeclareEnum(const char* Scope, const char* Name
 
 void FTestBatchBuilder::AddObject(FStructSchemaId Schema, FMemberBuilder&& Members)
 {
-	Objects.Emplace(Schema, Members.BuildAndReset(Scratch, Declarations.Get(Schema), *this));
+	Objects.Emplace(Schema, Members.BuildAndReset(Scratch, Declarations.Get(Schema), GetDebug()));
 }
 
 TArray64<uint8> FTestBatchBuilder::Write()
 {
 	// Build partial schemas
-	FSchemasBuilder SchemaBuilders(Declarations, *this, Scratch, ESchemaFormat::StableNames);
+	FSchemasBuilder SchemaBuilders(Declarations, *this, *this, Scratch, ESchemaFormat::StableNames);
 	for (const TPair<FStructSchemaId, FBuiltStruct*>& Object : Objects)
 	{
 		SchemaBuilders.NoteStructAndMembers(Object.Key, *Object.Value);
@@ -442,7 +444,7 @@ private:
 	TArray<const char*> Names;
 };
 
-class FTestBatchNames : public FDebugIds
+class FTestBatchNames : public FIds
 {
 public:
 	FTestBatchNames(const FTestNameReader& InNames, const FSchemaBatch& InSchemas) : Names(InNames), Schemas(InSchemas) {}
@@ -452,7 +454,7 @@ public:
 	virtual FTypeId					Resolve(FEnumSchemaId Id) const override final		{ return ResolveEnumSchema(Schemas, Id).Type; }
 	virtual FTypeId					Resolve(FStructSchemaId Id) const override final	{ return ResolveStructSchema(Schemas, Id).Type; }
 
-	using FDebugIds::AppendString;
+	using FIds::AppendString;
 	virtual void AppendString(FUtf8StringBuilderBase& Out, FNameId Name) const override final { Out.Append(Names[Name]); }
 
 private:
@@ -854,14 +856,14 @@ TEST_CASE_NAMED(FPlainPropsReadWriteTest, "System::Core::Serialization::PlainPro
 
 			FMemberBuilder Members;
 			Members.Add(Batch.NameMember("I1"), 100);
-			FBuiltStruct* NestedInStruct = Members.BuildAndReset(Scratch, Batch.Get(NestedId), Batch);
+			FBuiltStruct* NestedInStruct = Members.BuildAndReset(Scratch, Batch.Get(NestedId), Batch.GetDebug());
 
 			Members.AddStruct(Batch.NameMember("Nested"), NestedId, MoveTemp(NestedInStruct));
 			Members.Add(Batch.NameMember("Leaf"), true);
-			FBuiltStruct* Struct = Members.BuildAndReset(Scratch, Batch.Get(StructId), Batch);
+			FBuiltStruct* Struct = Members.BuildAndReset(Scratch, Batch.Get(StructId), Batch.GetDebug());
 		
 			Members.Add(Batch.NameMember("I2"), 200);
-			FBuiltStruct* NestedInObject = Members.BuildAndReset(Scratch, Batch.Get(NestedId), Batch);
+			FBuiltStruct* NestedInObject = Members.BuildAndReset(Scratch, Batch.Get(NestedId), Batch.GetDebug());
 
 			Members.Add(Batch.NameMember("L1"), 123.f);
 			Members.AddStruct(Batch.NameMember("S"), StructId, MoveTemp(Struct));
@@ -1099,7 +1101,7 @@ TEST_CASE_NAMED(FPlainPropsReadWriteTest, "System::Core::Serialization::PlainPro
 			Structs[2].Add(Batch.NameMember("F"), 2.f);
 		
 			FMemberBuilder Members;
-			Members.AddRange(Batch.NameMember("Structs"), Structs.BuildAndReset(Scratch, Batch.Get(StructId), Batch));
+			Members.AddRange(Batch.NameMember("Structs"), Structs.BuildAndReset(Scratch, Batch.Get(StructId), Batch.GetDebug()));
 
 			Batch.AddObject(ObjectId, MoveTemp(Members));
 		}, 
@@ -1160,7 +1162,7 @@ TEST_CASE_NAMED(FPlainPropsReadWriteTest, "System::Core::Serialization::PlainPro
 			XYs[1].Add(Batch.NameMember("X"), 3.f);
 			XYs[1].Add(Batch.NameMember("Y"), 4.f);
 			FNestedRangeBuilder StructRs(MakeStructRangeSchema(ERangeSizeType::U64, XY), 1);
-			StructRs.Add(XYs.BuildAndReset(Scratch, Batch.Get(XY), Batch));
+			StructRs.Add(XYs.BuildAndReset(Scratch, Batch.Get(XY), Batch.GetDebug()));
 
 			FStructRangeBuilder ZWs(int16(3));
 			ZWs[0].Add(Batch.NameMember("Z"), 1.5f);
@@ -1168,7 +1170,7 @@ TEST_CASE_NAMED(FPlainPropsReadWriteTest, "System::Core::Serialization::PlainPro
 			ZWs[2].Add(Batch.NameMember("W"), 3.5f);
 			FMemberSchema ZWRangeSchema = MakeStructRangeSchema(ERangeSizeType::S16, ZW);
 			FNestedRangeBuilder ZWRs(ZWRangeSchema, 1);
-			ZWRs.Add(ZWs.BuildAndReset(Scratch, Batch.Get(ZW), Batch));
+			ZWRs.Add(ZWs.BuildAndReset(Scratch, Batch.Get(ZW), Batch.GetDebug()));
 			FNestedRangeBuilder StructRRs(MakeNestedRangeSchema(Scratch, ERangeSizeType::U32, ZWRangeSchema), 1);
 			StructRRs.Add(ZWRs.BuildAndReset(Scratch, ERangeSizeType::U32));
 
@@ -1268,7 +1270,7 @@ TEST_CASE_NAMED(FPlainPropsReadWriteTest, "System::Core::Serialization::PlainPro
 
 			FMemberBuilder Members;
 			Members.AddRange(Batch.NameMember("Bools"), BuildLeafRange(Scratch, &True, true));
-			Members.AddRange(Batch.NameMember("Structs"), Structs.BuildAndReset(Scratch, Batch.Get(Struct), Batch));
+			Members.AddRange(Batch.NameMember("Structs"), Structs.BuildAndReset(Scratch, Batch.Get(Struct), Batch.GetDebug()));
 			Members.Add(Batch.NameMember("BF"), false);
 			Members.Add(Batch.NameMember("BT"), true);
 
@@ -1317,7 +1319,7 @@ TEST_CASE_NAMED(FPlainPropsReadWriteTest, "System::Core::Serialization::PlainPro
 				{
 					FMemberBuilder Members;
 					Members.Add(Batch.NameMember("X"), X);
-					return Members.BuildAndReset(Scratch, Batch.Get(Struct), Batch);
+					return Members.BuildAndReset(Scratch, Batch.Get(Struct), Batch.GetDebug());
 				};
 
 			FMemberBuilder O1;
@@ -1364,7 +1366,7 @@ TEST_CASE_NAMED(FPlainPropsReadWriteTest, "System::Core::Serialization::PlainPro
 				{
 					FStructRangeBuilder Members(1);
 					Members[0].Add(Batch.NameMember("X"), X);
-					return Members.BuildAndReset(Scratch, Batch.Get(Struct), Batch);
+					return Members.BuildAndReset(Scratch, Batch.Get(Struct), Batch.GetDebug());
 				};
 
 			
@@ -1381,8 +1383,8 @@ TEST_CASE_NAMED(FPlainPropsReadWriteTest, "System::Core::Serialization::PlainPro
 			FMemberBuilder O2;
 			O2.AddRange(Batch.NameMember("Same"), BuildStructRange(SA, 20));
 			O2.AddRange(Batch.NameMember("Diff"), BuildStructRange(SB, 22.f));
-			O2.AddRange(Batch.NameMember("SameEmpty"), FStructRangeBuilder(0).BuildAndReset(Scratch, Batch.Get(SA), Batch));
-			O2.AddRange(Batch.NameMember("DiffEmpty"), FStructRangeBuilder(0).BuildAndReset(Scratch, Batch.Get(SB), Batch));
+			O2.AddRange(Batch.NameMember("SameEmpty"), FStructRangeBuilder(0).BuildAndReset(Scratch, Batch.Get(SA), Batch.GetDebug()));
+			O2.AddRange(Batch.NameMember("DiffEmpty"), FStructRangeBuilder(0).BuildAndReset(Scratch, Batch.Get(SB), Batch.GetDebug()));
 			FNestedRangeBuilder NestedSB(MakeStructRangeSchema(ERangeSizeType::S32, SB), 1);
 			NestedSB.Add(BuildStructRange(SB, 200.f));
 			O2.AddRange(Batch.NameMember("DiffNested"), NestedSB.BuildAndReset(Scratch, ERangeSizeType::S32));
@@ -1424,20 +1426,20 @@ TEST_CASE_NAMED(FPlainPropsReadWriteTest, "System::Core::Serialization::PlainPro
 			
 			FMemberBuilder Members;
 			Members.Add(Batch.NameMember("LInt"), 123);
-			Members.BuildSuperStruct(Scratch, Batch.Get(Low), Batch);
+			Members.BuildSuperStruct(Scratch, Batch.Get(Low), Batch.GetDebug());
 			Members.Add(Batch.NameMember("MInt"), 456);
 			FMemberBuilder Nested;
 			Nested.Add(Batch.NameMember("LInt"), 1000);
-			Members.AddStruct(Batch.NameMember("MLow"), Low, Nested.BuildAndReset(Scratch, Batch.Get(Low), Batch));
-			Members.BuildSuperStruct(Scratch, Batch.Get(Mid), Batch);
+			Members.AddStruct(Batch.NameMember("MLow"), Low, Nested.BuildAndReset(Scratch, Batch.Get(Low), Batch.GetDebug()));
+			Members.BuildSuperStruct(Scratch, Batch.Get(Mid), Batch.GetDebug());
 			Members.Add(Batch.NameMember("TInt"), 789);
 			Nested.Add(Batch.NameMember("LInt"), 2000);
-			Members.AddStruct(Batch.NameMember("TLow"), Low, Nested.BuildAndReset(Scratch, Batch.Get(Low), Batch));
+			Members.AddStruct(Batch.NameMember("TLow"), Low, Nested.BuildAndReset(Scratch, Batch.Get(Low), Batch.GetDebug()));
 			FStructRangeBuilder NestedRange(1);
 			NestedRange[0].Add(Batch.NameMember("MInt"), 3000);
 			Nested.Add(Batch.NameMember("LInt"), 4000);
-			NestedRange[0].AddStruct(Batch.NameMember("MLow"), Low, Nested.BuildAndReset(Scratch, Batch.Get(Low), Batch));
-			Members.AddRange(Batch.NameMember("TMids"), NestedRange.BuildAndReset(Scratch, Batch.Get(Mid), Batch));
+			NestedRange[0].AddStruct(Batch.NameMember("MLow"), Low, Nested.BuildAndReset(Scratch, Batch.Get(Low), Batch.GetDebug()));
+			Members.AddRange(Batch.NameMember("TMids"), NestedRange.BuildAndReset(Scratch, Batch.Get(Mid), Batch.GetDebug()));
 
 			Batch.AddObject(Top, MoveTemp(Members));
 		}, 
@@ -1509,29 +1511,29 @@ TEST_CASE_NAMED(FPlainPropsReadWriteTest, "System::Core::Serialization::PlainPro
 			
 			FMemberBuilder A;
 			A.Add(Batch.NameMember("1"), 1);
-			A.BuildSuperStruct(Scratch, Batch.Get(B1), Batch);
+			A.BuildSuperStruct(Scratch, Batch.Get(B1), Batch.GetDebug());
 			A.Add(Batch.NameMember("4"), 4);
-			A.BuildSuperStruct(Scratch, Batch.Get(B4), Batch);
+			A.BuildSuperStruct(Scratch, Batch.Get(B4), Batch.GetDebug());
 			A.Add(Batch.NameMember("5"), 5);
-			A.BuildSuperStruct(Scratch, Batch.Get(B5), Batch);
+			A.BuildSuperStruct(Scratch, Batch.Get(B5), Batch.GetDebug());
 
 			FMemberBuilder B;
 			B.Add(Batch.NameMember("1"), 10);
-			B.BuildSuperStruct(Scratch, Batch.Get(B1), Batch);
+			B.BuildSuperStruct(Scratch, Batch.Get(B1), Batch.GetDebug());
 			B.Add(Batch.NameMember("2"), 20);
-			B.BuildSuperStruct(Scratch, Batch.Get(B2), Batch);
+			B.BuildSuperStruct(Scratch, Batch.Get(B2), Batch.GetDebug());
 			B.Add(Batch.NameMember("4"), 40);
-			B.BuildSuperStruct(Scratch, Batch.Get(B4), Batch);
+			B.BuildSuperStruct(Scratch, Batch.Get(B4), Batch.GetDebug());
 
 			FMemberBuilder C;
 			C.Add(Batch.NameMember("1"), 100);
-			C.BuildSuperStruct(Scratch, Batch.Get(B1), Batch);
+			C.BuildSuperStruct(Scratch, Batch.Get(B1), Batch.GetDebug());
 			C.Add(Batch.NameMember("2"), 200);
-			C.BuildSuperStruct(Scratch, Batch.Get(B2), Batch);
-			C.BuildSuperStruct(Scratch, Batch.Get(B3), Batch); // Empty -> noop
+			C.BuildSuperStruct(Scratch, Batch.Get(B2), Batch.GetDebug());
+			C.BuildSuperStruct(Scratch, Batch.Get(B3), Batch.GetDebug()); // Empty -> noop
 			C.Add(Batch.NameMember("4"), 400);
-			C.BuildSuperStruct(Scratch, Batch.Get(B4), Batch);
-			C.BuildSuperStruct(Scratch, Batch.Get(B5), Batch); // Empty -> noop
+			C.BuildSuperStruct(Scratch, Batch.Get(B4), Batch.GetDebug());
+			C.BuildSuperStruct(Scratch, Batch.Get(B5), Batch.GetDebug()); // Empty -> noop
 			
 			Batch.AddObject(B6, MoveTemp(A));
 			Batch.AddObject(B5, MoveTemp(B));
diff --git a/Engine/Plugins/Experimental/PlainProps/Source/Tests/PlainPropsUeCoreTests.cpp b/Engine/Plugins/Experimental/PlainProps/Source/Tests/PlainPropsUeCoreTests.cpp
--- a/Engine/Plugins/Experimental/PlainProps/Source/Tests/PlainPropsUeCoreTests.cpp
+++ b/Engine/Plugins/Experimental/PlainProps/Source/Tests/PlainPropsUeCoreTests.cpp
@@ -31,12 +31,13 @@ namespace PlainProps::UE::Test
 {
 
 static TIdIndexer<FName>		GNames;
-static FDeclarations			GTypes(/* debug */ GNames);
-static FSchemaBindings			GSchemas(/* debug */ GNames);
-static FCustomBindingsBottom	GCustoms(/* debug */ GNames);
-static FCustomBindingsOverlay	GDeltaCustoms(/* base*/ GCustoms);
+static FDebugIds				GDebug(GNames);
+static FDeclarations			GTypes(GDebug);
+static FSchemaBindings			GSchemas(GDebug);
+static FCustomBindingsBottom	GCustoms(GDebug);
+static FCustomBindingsOverlay	GDeltaCustoms(GCustoms);
 
-struct FIds
+struct FRuntimeIds
 {
 	static FNameId				IndexName(FAnsiStringView Name)		{ return GNames.MakeName(FName(Name)); }
 	static FMemberId			IndexMember(FAnsiStringView Name)	{ return GNames.NameMember(FName(Name)); }
@@ -45,14 +46,13 @@ struct FIds
 	static FEnumSchemaId		IndexEnum(FTypeId Type)				{ return GNames.IndexEnum(Type); }
 	static FStructSchemaId		IndexStruct(FTypeId Type)			{ return GNames.IndexStruct(Type); }
 	static FIdIndexerBase&		GetIndexer()						{ return GNames; }
-	static const FDebugIds&		GetDebug()							{ return GNames; }
 };
 
 ////////////////////////////////////////////////////////////////////////////////////////////////
 
 struct FDefaultRuntime
 {
-	using Ids = FIds;
+	using Ids = FRuntimeIds;
 	template<class T> using CustomBindings = TCustomBind<T>;
 
 	static FDeclarations&			GetTypes()			{ return GTypes; }
@@ -87,15 +87,15 @@ using TScopedDefaultStructBinding = TScopedStructBinding<T, FDefaultRuntime>;
 
 ////////////////////////////////////////////////////////////////////////////////////////////////
 
-// Maybe replace with macro, e.g. PP_DECLARE_CUSTOM_DENSE_STRUCT(FIds, FName, void, Idx)
+// Maybe replace with macro, e.g. PP_DECLARE_CUSTOM_DENSE_STRUCT(FRuntimeIds, FName, void, Idx)
 struct FNameDeclaration 
 {
 	FStructSchemaId		Id;
 	FMemberId			Idx;
 
-	FNameDeclaration(FTypeId Type = IndexStructName<FIds, ETypename::Decl, TTypename<FName>>())
-	: Id(FIds::IndexStruct(Type))
-	, Idx(FIds::IndexMember("Idx"))
+	FNameDeclaration(FTypeId Type = IndexStructName<FRuntimeIds, ETypename::Decl, TTypename<FName>>())
+	: Id(FRuntimeIds::IndexStruct(Type))
+	, Idx(FRuntimeIds::IndexMember("Idx"))
 	{
 		GTypes.DeclareStruct(Id, Type, MakeArrayView(&Idx, 1), EMemberPresence::RequireAll);
 	}
@@ -169,14 +169,14 @@ FBatchSaver::FBatchSaver(const FCustomBindings& Underlay)
 template<class T>
 void FBatchSaver::Save(T&& Object) 
 {
-	FDualStructSchemaId Id = IndexStructDualId<FIds, TTypename<std::remove_reference_t<T>>>();
+	FDualStructSchemaId Id = IndexStructDualId<FRuntimeIds, TTypename<std::remove_reference_t<T>>>();
 	SavedObjects.Emplace(Id.DeclId, SaveStruct(&Object, Id.BindId, {GTypes, GSchemas, Customs, Scratch}));
 }
 
 template<class T>
 bool FBatchSaver::SaveDelta(const T& Object, const T& Default) 
 {
-	FDualStructSchemaId Id = IndexStructDualId<FIds, TTypename<std::remove_reference_t<T>>>();
+	FDualStructSchemaId Id = IndexStructDualId<FRuntimeIds, TTypename<std::remove_reference_t<T>>>();
 	if (const FBuiltStruct* Delta = SaveStructDeltaIfDiff(&Object, &Default, Id.BindId, {GTypes, GSchemas, Customs, Scratch}))
 	{
 		SavedObjects.Emplace(Id.DeclId, MoveTemp(Delta));
@@ -207,7 +207,7 @@ TArray64<uint8> FBatchSaver::Write(TArray<FStructSchemaId>* OutMemoryIds) const
 
 	// Build partial schemas
 	const FStructBindIds BindIds(Customs, GSchemas);
-	FSchemasBuilder SchemaBuilders(GTypes, BindIds, Scratch, Format);
+	FSchemasBuilder SchemaBuilders(GTypes, GNames, BindIds, Scratch, Format);
 	for (const IdBuiltStructPair& Object : SavedObjects)
 	{
 		SchemaBuilders.NoteStructAndMembers(Object.Key, *Object.Value);
@@ -264,7 +264,7 @@ TArray64<uint8> FBatchSaver::Write(TArray<FStructSchemaId>* OutMemoryIds) const
 
 ////////////////////////////////////////////////////////////////////////////////////////////////
 
-class FSchemaBatchIds : public FDebugIds
+class FSchemaBatchIds : public FIds
 {
 public:
 	FSchemaBatchIds(const FSchemaBatch& InSchemas) : Schemas(InSchemas) {}
@@ -300,23 +300,23 @@ public:
 	TConstArrayView<FName> Names;
 };
 
-class FLoadBatchNames : public FDebugIds
+class FLoadBatchNames : public FIds
 {
 public:
-	FLoadBatchNames(const FIdBinding& InBinding, const FSchemaBatch& InSchemas) : DebugIds(GNames), Binding(InBinding), Schemas(InSchemas) {}
+	FLoadBatchNames(const FIdBinding& InBinding, const FSchemaBatch& InSchemas) : Ids(GNames), Binding(InBinding), Schemas(InSchemas) {}
 
-	virtual FNestedScope		Resolve(FNestedScopeId Id) const override final		{ return DebugIds.Resolve(Binding.Remap(Id)); }
-	virtual FParametricTypeView	Resolve(FParametricTypeId Id) const override final	{ return DebugIds.Resolve(Binding.Remap(Id)); }
+	virtual FNestedScope		Resolve(FNestedScopeId Id) const override final		{ return Ids.Resolve(Binding.Remap(Id)); }
+	virtual FParametricTypeView	Resolve(FParametricTypeId Id) const override final	{ return Ids.Resolve(Binding.Remap(Id)); }
 	virtual FTypeId				Resolve(FEnumSchemaId Id) const override final		{ return ResolveEnumSchema(Schemas, Id).Type; }
 	virtual FTypeId				Resolve(FStructSchemaId Id) const override final	{ return ResolveStructSchema(Schemas, Id).Type; }
 
-	using FDebugIds::AppendString;
-	virtual void AppendString(FUtf8Builder& Out, FNameId Name) const override final			{ DebugIds.AppendString(Out, Name); }
-	virtual void AppendString(FUtf8Builder& Out, FTypenameId Typename) const override final	{ DebugIds.AppendString(Out, Typename); }
-	virtual void AppendString(FUtf8Builder& Out, FScopeId Scope) const override final		{ DebugIds.AppendString(Out, Scope); }
+	using FIds::AppendString;
+	virtual void AppendString(FUtf8Builder& Out, FNameId Name) const override final			{ Ids.AppendString(Out, Name); }
+	virtual void AppendString(FUtf8Builder& Out, FTypenameId Typename) const override final	{ Ids.AppendString(Out, Typename); }
+	virtual void AppendString(FUtf8Builder& Out, FScopeId Scope) const override final		{ Ids.AppendString(Out, Scope); }
 
 private:
-	const FDebugIds& DebugIds;
+	const FIds& Ids;
 	const FIdBinding& Binding;
 	const FSchemaBatch& Schemas;
 };
diff --git a/Engine/Plugins/Experimental/PlainPropsUObject/Source/Private/PlainPropsUObjectBindings.cpp b/Engine/Plugins/Experimental/PlainPropsUObject/Source/Private/PlainPropsUObjectBindings.cpp
--- a/Engine/Plugins/Experimental/PlainPropsUObject/Source/Private/PlainPropsUObjectBindings.cpp
+++ b/Engine/Plugins/Experimental/PlainPropsUObject/Source/Private/PlainPropsUObjectBindings.cpp
@@ -167,7 +167,7 @@ void FDefaultStructs::Bind(FStructSchemaId Id, const UScriptStruct* Struct)
 
 	ReserveFlags(Id.Idx);
 #if DO_CHECK
-	checkf(Flip(Bound[Id.Idx]), TEXT("'%s' already bound"), *GUE.Names.Print(Id));
+	checkf(Flip(Bound[Id.Idx]), TEXT("'%s' already bound"), *GUE.Debug.Print(Id));
 #endif
 
 	if (const UUserDefinedStruct* UserStruct = Cast<UUserDefinedStruct>(Struct))
@@ -209,7 +209,7 @@ void FDefaultStructs::BindZeroes(FStructSchemaId Id, SIZE_T Size, uint32 Alignme
 {
 	ReserveFlags(Id.Idx);
 #if DO_CHECK
-	checkf(Flip(Bound[Id.Idx]), TEXT("'%s' already bound"), *GUE.Names.Print(Id));
+	checkf(Flip(Bound[Id.Idx]), TEXT("'%s' already bound"), *GUE.Debug.Print(Id));
 #endif
 	ReserveZeroes(Zeroes, Size, Alignment);
 }
@@ -218,7 +218,7 @@ void FDefaultStructs::BindStatic(FStructSchemaId Id, const void* Struct)
 {
 	ReserveFlags(Id.Idx);
 #if DO_CHECK
-	checkf(Flip(Bound[Id.Idx]), TEXT("'%s' already bound"), *GUE.Names.Print(Id));
+	checkf(Flip(Bound[Id.Idx]), TEXT("'%s' already bound"), *GUE.Debug.Print(Id));
 #endif
 	check(!Instanced[Id.Idx]);
 	check(GetInstance(MakeStaticInstance(Struct)) == Struct);
@@ -230,7 +230,7 @@ void FDefaultStructs::BindStatic(FStructSchemaId Id, const void* Struct)
 void FDefaultStructs::Drop(FStructSchemaId Id)
 {
 #if DO_CHECK
-	checkf(!Flip(Bound[Id.Idx]), TEXT("'%s' isn't bound"), *GUE.Names.Print(Id));
+	checkf(!Flip(Bound[Id.Idx]), TEXT("'%s' isn't bound"), *GUE.Debug.Print(Id));
 #endif
 	if (Instanced[Id.Idx])
 	{
@@ -242,7 +242,7 @@ void FDefaultStructs::Drop(FStructSchemaId Id)
 const void* FDefaultStructs::Get(FStructSchemaId Id)
 {
 #if DO_CHECK
-	checkf(Bound[Id.Idx], TEXT("'%s' lack default"), *GUE.Names.Print(Id));
+	checkf(Bound[Id.Idx], TEXT("'%s' lack default"), *GUE.Debug.Print(Id));
 #endif
 	return Instanced[Id.Idx] ? GetInstance(Instances.FindChecked(Id)) : Zeroes.GetData(); 
 }
@@ -326,14 +326,15 @@ FCommonMemberIds::FCommonMemberIds(TIdIndexer<FName>& Names)
 ////////////////////////////////////////////////////////////////////////////////////////////////
 
 FGlobals::FGlobals() 
-: Types(Names)
-, Schemas(Names)
-, Customs(Names)
+: Types(FDebugIds(Names))
+, Schemas(FDebugIds(Names))
+, Customs(FDebugIds(Names))
 , DeltaCustoms(Customs)
 , Scopes(Names)
 , Structs(Scopes, Names)
 , Typenames(Names)
 , Members(Names)
+, Debug(Names)
 {}
 
 FGlobals GUE;
@@ -630,7 +631,7 @@ static FTypedRange SaveInvocations(const FMulticastInvocationList& In, const FSa
 				if (Keep[Idx])
 				{
 					SaveUnicastDelegate(/* out */ Tmp, Src[Idx], Ctx);
-					Dst.AddItem(Tmp.BuildAndReset(Ctx.Scratch, ItemDecl, GUE.Names));
+					Dst.AddItem(Tmp.BuildAndReset(Ctx.Scratch, ItemDecl, GUE.Debug));
 				}
 			}
 			return Dst.Finalize(Schema);
@@ -2977,7 +2978,7 @@ private:
 		if (Inner.InnermostType.IsLeaf() && Inner.InnermostType.AsLeaf().Bind.Type == ELeafBindType::BitfieldBool)
 		{
 			UE_LOGFMT(LogPlainPropsUObject, Warning, "Property '{Property}' is a '{Container}' of bitfield bools, which make no sense. Binding as range of bools.", 
-				Property->GetFName(), GUE.Names.Print(RangeBinding.GetBindName()));
+				Property->GetFName(), GUE.Debug.Print(RangeBinding.GetBindName()));
 			Inner.InnermostType = FMemberBindType(ReflectArithmetic<bool>);
 		}
 
@@ -3517,7 +3518,7 @@ public:
 	{
 		// Build partial schemas
 		const FStructBindIds BindIds(GUE.DeltaCustoms, GUE.Schemas);
-		FSchemasBuilder SchemaBuilders(GUE.Types, BindIds, Scratch, ESchemaFormat::InMemoryNames);
+		FSchemasBuilder SchemaBuilders(GUE.Types, GUE.Names, BindIds, Scratch, ESchemaFormat::InMemoryNames);
 		for (const FSavedObject& Object : SavedObjects)
 		{
 			GDebugNoteObject = Object.Input;
@@ -4009,12 +4010,12 @@ int32 UTestPlainPropsCommandlet::Main(const FString& Params)
 }
 
 // Stopgap debugger callable functions til natvis is added 
-extern FString _print(const PlainProps::FNameId& Id)				{ return PlainProps::UE::GUE.Names.Print(Id); }
-extern FString _print(const PlainProps::FMemberId& Id)				{ return PlainProps::UE::GUE.Names.Print(Id); }
-extern FString _print(const PlainProps::FOptionalMemberId& Id)		{ return PlainProps::UE::GUE.Names.Print(Id); }
-extern FString _print(const PlainProps::FScopeId& Id)				{ return PlainProps::UE::GUE.Names.Print(Id); }
-extern FString _print(const PlainProps::FTypenameId& Id)			{ return PlainProps::UE::GUE.Names.Print(Id); }
-extern FString _print(const PlainProps::FConcreteTypenameId& Id)	{ return PlainProps::UE::GUE.Names.Print(Id); }
-extern FString _print(const PlainProps::FTypeId& Id)				{ return PlainProps::UE::GUE.Names.Print(Id); }
-extern FString _print(const PlainProps::FEnumSchemaId& Id)			{ return PlainProps::UE::GUE.Names.Print(Id); }
-extern FString _print(const PlainProps::FStructSchemaId& Id)		{ return PlainProps::UE::GUE.Names.Print(Id); }
+extern FString _print(const PlainProps::FNameId& Id)				{ return PlainProps::UE::GUE.Debug.Print(Id); }
+extern FString _print(const PlainProps::FMemberId& Id)				{ return PlainProps::UE::GUE.Debug.Print(Id); }
+extern FString _print(const PlainProps::FOptionalMemberId& Id)		{ return PlainProps::UE::GUE.Debug.Print(Id); }
+extern FString _print(const PlainProps::FScopeId& Id)				{ return PlainProps::UE::GUE.Debug.Print(Id); }
+extern FString _print(const PlainProps::FTypenameId& Id)			{ return PlainProps::UE::GUE.Debug.Print(Id); }
+extern FString _print(const PlainProps::FConcreteTypenameId& Id)	{ return PlainProps::UE::GUE.Debug.Print(Id); }
+extern FString _print(const PlainProps::FTypeId& Id)				{ return PlainProps::UE::GUE.Debug.Print(Id); }
+extern FString _print(const PlainProps::FEnumSchemaId& Id)			{ return PlainProps::UE::GUE.Debug.Print(Id); }
+extern FString _print(const PlainProps::FStructSchemaId& Id)		{ return PlainProps::UE::GUE.Debug.Print(Id); }
diff --git a/Engine/Plugins/Experimental/PlainPropsUObject/Source/Public/PlainPropsUObjectRuntime.h b/Engine/Plugins/Experimental/PlainPropsUObject/Source/Public/PlainPropsUObjectRuntime.h
--- a/Engine/Plugins/Experimental/PlainPropsUObject/Source/Public/PlainPropsUObjectRuntime.h
+++ b/Engine/Plugins/Experimental/PlainPropsUObject/Source/Public/PlainPropsUObjectRuntime.h
@@ -145,10 +145,11 @@ struct FGlobals
 	FCommonTypenameIds			Typenames;
 	FCommonMemberIds			Members;
 	FNumeralGenerator			Numerals;
+	FDebugIds					Debug;
 };
 PLAINPROPSUOBJECT_API extern FGlobals GUE;
 
-struct FIds
+struct FRuntimeIds
 {
 	static FNameId				IndexName(FAnsiStringView Name)			{ return GUE.Names.MakeName(FName(Name)); }
 	static FMemberId			IndexMember(FAnsiStringView Name)		{ return GUE.Names.NameMember(FName(Name)); }
@@ -157,12 +158,11 @@ struct FIds
 	static FEnumSchemaId		IndexEnum(FTypeId Type)					{ return GUE.Names.IndexEnum(Type); }
 	static FStructSchemaId		IndexStruct(FTypeId Type)				{ return GUE.Names.IndexStruct(Type); }
 	static FIdIndexerBase&		GetIndexer()							{ return GUE.Names; }
-	static const FDebugIds&		GetDebug()								{ return GUE.Names; }
 };
 
 struct FDefaultRuntime
 {
-	using Ids = FIds;
+	using Ids = FRuntimeIds;
 	template<class T> using CustomBindings = TCustomBind<T>;
 
 	static FDeclarations&		GetTypes()			{ return GUE.Types; }
@@ -187,7 +187,7 @@ struct FFieldPathBinding : ICustomBinding
 	using Type = FFieldPath;
 	const FMemberId MemberIds[2];
 
-	FFieldPathBinding(TCustomInit<FIds>) : MemberIds{GUE.Members.Path, GUE.Members.Owner} {}
+	FFieldPathBinding(TCustomInit<FRuntimeIds>) : MemberIds{GUE.Members.Path, GUE.Members.Owner} {}
 
 	void			Save(FMemberBuilder& Dst, const Type& Src, const Type* Default, const FSaveContext&) const;
 	void			Load(Type& Dst, FStructLoadView Src, ECustomLoadMethod Method) const;
@@ -199,7 +199,7 @@ struct FDelegateBinding : ICustomBinding
 	using Type = FScriptDelegate;
 	const FMemberId MemberIds[2];
 
-	FDelegateBinding(TCustomInit<FIds>) : MemberIds{GUE.Members.Object, GUE.Members.Function} {}
+	FDelegateBinding(TCustomInit<FRuntimeIds>) : MemberIds{GUE.Members.Object, GUE.Members.Function} {}
 
 	void			Save(FMemberBuilder& Dst, const Type& Src, const Type* Default, const FSaveContext&) const;
 	void			Load(Type& Dst, FStructLoadView Src, ECustomLoadMethod Method) const;
@@ -211,7 +211,7 @@ struct FMulticastInlineDelegateBinding : ICustomBinding
 	using Type = FMulticastScriptDelegate;
 	const FMemberId MemberIds[1];
 
-	FMulticastInlineDelegateBinding(TCustomInit<FIds>) : MemberIds{GUE.Members.Invocations} {}
+	FMulticastInlineDelegateBinding(TCustomInit<FRuntimeIds>) : MemberIds{GUE.Members.Invocations} {}
 
 	void			Save(FMemberBuilder& Dst, const Type& Src, const Type* Default, const FSaveContext&) const;
 	void			Load(Type& Dst, FStructLoadView Src, ECustomLoadMethod Method) const;
@@ -225,7 +225,7 @@ struct FVerseFunctionBinding : ICustomBinding
 	using Type =  FVerseFunction;
 	const FMemberId MemberIds[1];
 
-	FVerseFunctionBinding(TCustomInit<FIds>) : MemberIds{GUE.Members.Value} {}
+	FVerseFunctionBinding(TCustomInit<FRuntimeIds>) : MemberIds{GUE.Members.Value} {}
 
 	void			Save(FMemberBuilder& Dst, const Type& Src, const Type* Default, const FSaveContext&) const;
 	void			Load(Type& Dst, FStructLoadView Src, ECustomLoadMethod Method) const;
@@ -237,7 +237,7 @@ struct FDynamicallyTypedValueBinding : ICustomBinding
 	using Type = ::UE::FDynamicallyTypedValue;
 	const FMemberId MemberIds[1];
 
-	FDynamicallyTypedValueBinding(TCustomInit<FIds>) : MemberIds{GUE.Members.Value} {}
+	FDynamicallyTypedValueBinding(TCustomInit<FRuntimeIds>) : MemberIds{GUE.Members.Value} {}
 
 	void			Save(FMemberBuilder& Dst, const Type& Src, const Type* Default, const FSaveContext&) const;
 	void			Load(Type& Dst, FStructLoadView Src, ECustomLoadMethod Method) const;
@@ -250,7 +250,7 @@ struct FReferencePropertyBinding : ICustomBinding
 	using Type = FReferencePropertyValue;
 	const FMemberId MemberIds[1];
 
-	FReferencePropertyBinding(TCustomInit<FIds>) : MemberIds{GUE.Members.Value} {}
+	FReferencePropertyBinding(TCustomInit<FRuntimeIds>) : MemberIds{GUE.Members.Value} {}
 
 	void			Save(FMemberBuilder& Dst, const Type& Src, const Type* Default, const FSaveContext&) const;
 	void			Load(Type& Dst, FStructLoadView Src, ECustomLoadMethod Method) const;
@@ -269,9 +269,9 @@ template<> struct TCustomBind<FReferencePropertyValue> { using Type = UE::FRefer
 template<> struct TCustomBind<FVerseFunction> { using Type = UE::FVerseFunctionBinding; };
 
 // Temporary way to tie certain types to /Script/Core scope
-template<> FScopeId IndexNamespaceId<UE::FIds, TTypename<FTransform>>()		{ return UE::GUE.Scopes.Core; }
-template<> FScopeId IndexNamespaceId<UE::FIds, TTypename<FGuid>>()			{ return UE::GUE.Scopes.Core; }
-template<> FScopeId IndexNamespaceId<UE::FIds, TTypename<FColor>>()			{ return UE::GUE.Scopes.Core; }
-template<> FScopeId IndexNamespaceId<UE::FIds, TTypename<FLinearColor>>()	{ return UE::GUE.Scopes.Core; }
+template<> FScopeId IndexNamespaceId<UE::FRuntimeIds, TTypename<FTransform>>()		{ return UE::GUE.Scopes.Core; }
+template<> FScopeId IndexNamespaceId<UE::FRuntimeIds, TTypename<FGuid>>()			{ return UE::GUE.Scopes.Core; }
+template<> FScopeId IndexNamespaceId<UE::FRuntimeIds, TTypename<FColor>>()			{ return UE::GUE.Scopes.Core; }
+template<> FScopeId IndexNamespaceId<UE::FRuntimeIds, TTypename<FLinearColor>>()	{ return UE::GUE.Scopes.Core; }
 
-}
\ No newline at end of file
+}
