PlainProps implement memcopy load plan and optimized FString loading

#rnx

diff --git a/Engine/Plugins/Experimental/PlainProps/Source/Private/PlainPropsLoad.cpp b/Engine/Plugins/Experimental/PlainProps/Source/Private/PlainPropsLoad.cpp
--- a/Engine/Plugins/Experimental/PlainProps/Source/Private/PlainPropsLoad.cpp
+++ b/Engine/Plugins/Experimental/PlainProps/Source/Private/PlainPropsLoad.cpp
@@ -6,6 +6,7 @@
 #include "PlainPropsInternalBind.h"
 #include "PlainPropsInternalFormat.h"
 #include "PlainPropsInternalRead.h"
+#include "Algo/Compare.h"
 #include "Algo/Find.h"
 #include "Containers/Set.h"
 #include "Misc/Optional.h"
@@ -67,7 +68,7 @@ public:
 	bool						IsSparseSchema() const	{ return (Handle & SparseSchemaMask) == SparseSchemaMask; }
 	bool						IsMemcpy() const		{ return (Handle & LoMask) == MemcpyMask; }
 	bool						IsCustom() const		{ return (Handle & LoMask) == CustomMask; }
-	FLoadStructMemcpy			AsMemcpy() const		{ check(IsMemcpy()); return { static_cast<uint32>(Handle >> 32), static_cast<uint32>(Handle) >> 1 }; }
+	FLoadStructMemcpy			AsMemcpy() const		{ check(IsMemcpy()); return { static_cast<uint32>(Handle >> 32), static_cast<uint32>(Handle) >> 2 }; }
 	const FSchemaBinding&		AsSchema() const		{ check(IsSchema()); return *AsPtr<FSchemaBinding>(); }
 	const ICustomBinding&		AsCustom() const;
 
@@ -401,10 +402,42 @@ template<class LoadIdBinder>
 	return FLoadStructPlan(To, ELeafWidth::B32, !From.IsDense);
 }
 
-[[nodiscard]] static TOptional<FLoadStructMemcpy> TryMakeMemcpyPlan(const FStructSchema& From, const FSchemaBinding& To, TConstArrayView<FMemberId> ToMemberIds)
+[[nodiscard]] static TOptional<FLoadStructMemcpy> TryMakeMemcpyPlan(const FStructSchema& From, const FSchemaBinding& To, TConstArrayView<FMemberId> ToNames)
 {
-	// todo
-	return NullOpt;
+	// Can't memcopy sparse members nor range-bound members nor super structs
+	if (!From.IsDense || From.NumRangeTypes > 0 || UsesSuper(From.Inheritance))
+	{
+		return NullOpt;
+	}
+
+	// Can't memcpy non-contiguous members
+	TConstArrayView<FMemberId> FromNames = From.GetMemberNames();
+	int32 SkipToIdx = ToNames.Find(FromNames[0]);
+	check(SkipToIdx != INDEX_NONE);
+	if (!Algo::Compare(FromNames, ToNames.Slice(SkipToIdx, FromNames.Num())))
+	{	
+		return NullOpt;
+	}
+	
+	// Check all schema members are contiguous leaves
+	const uint32* OffsetIt = To.GetOffsets() + SkipToIdx;
+	const uint32 StartPos = *OffsetIt;
+	uint32 EndPos = *OffsetIt;
+	for (FMemberType Member : From.GetMemberTypes())
+	{
+		// Note: By adding FStructType::IsDense flag and some struct size lookup, memcpying of nested non-dynamic non-custom-bound structs might be possible
+		uint32 Offset = *OffsetIt++;
+		if (Offset != EndPos ||	// Non-contiguous or padded
+			Member.IsStruct() || // Nested dtructs have a skippable size prefix that can't be memcopied
+			Member.AsLeaf().Type == ELeafType::Bool) // bool values are stored as bits, see FBitCacheReader
+		{
+			return NullOpt;
+		}
+
+		EndPos += SizeOf(Member.AsLeaf().Width);
+	}
+
+	return FLoadStructMemcpy(EndPos - StartPos, StartPos);
 }
 
 template<class LoadIdBinder>
@@ -1093,8 +1126,10 @@ void LoadStruct(void* Dst, FByteReader Src, FStructSchemaId LoadId, const FLoadB
 	}
 	else if (Plan.IsMemcpy())
 	{
-		Src.CheckSize(Plan.AsMemcpy().Size);
-		FMemory::Memcpy(static_cast<uint8*>(Dst) + Plan.AsMemcpy().Offset, Src.Peek(), Plan.AsMemcpy().Size);
+		static_assert(PLATFORM_LITTLE_ENDIAN);
+		FLoadStructMemcpy MemcpyPlan = Plan.AsMemcpy();
+		Src.CheckSize(MemcpyPlan.Size);
+		FMemory::Memcpy(static_cast<uint8*>(Dst) + MemcpyPlan.Offset, Src.Peek(), MemcpyPlan.Size);
 	}
 	else
 	{
diff --git a/Engine/Plugins/Experimental/PlainProps/Source/Public/PlainPropsUeCoreBindings.h b/Engine/Plugins/Experimental/PlainProps/Source/Public/PlainPropsUeCoreBindings.h
--- a/Engine/Plugins/Experimental/PlainProps/Source/Public/PlainPropsUeCoreBindings.h
+++ b/Engine/Plugins/Experimental/PlainProps/Source/Public/PlainPropsUeCoreBindings.h
@@ -162,22 +162,33 @@ struct TStringBinding : ILeafRangeBinding
 	
 	virtual void SaveLeaves(const void* Range, FLeafRangeAllocator& Out) const override
 	{
-		const TArray<CharType>& Src = static_cast<const StringType*>(Range)->GetCharArray();
-		int32 SrcLen = Src.Num() - 1;
+		const TArray<CharType>& SrcArray = static_cast<const StringType*>(Range)->GetCharArray();
+		const CharType* Src = SrcArray.GetData();
+		int32 SrcLen = SrcArray.Num() - 1;
 		if (SrcLen <= 0)
 		{
 		}
 		else if constexpr (sizeof(CharType) == sizeof(char8_t))
 		{
 			char8_t* Utf8 = Out.AllocateRange<char8_t>(SrcLen);
-			FMemory::Memcpy(Utf8, Src.GetData(), SrcLen);
+			FMemory::Memcpy(Utf8, Src, SrcLen);
 		}
 		else
 		{
-			int32 Utf8Len = FPlatformString::ConvertedLength<UTF8CHAR>(Src.GetData(), SrcLen);
+			int32 Utf8Len = FPlatformString::ConvertedLength<UTF8CHAR>(Src, SrcLen);
 			char8_t* Utf8 = Out.AllocateRange<char8_t>(Utf8Len);
-			UTF8CHAR* Utf8End = FPlatformString::Convert(reinterpret_cast<UTF8CHAR*>(Utf8), Utf8Len, Src.GetData(), SrcLen);	
-			check((char8_t*)Utf8End - Utf8 == Utf8Len);
+			if (Utf8Len == SrcLen)
+			{
+				for (int32 Idx = 0; Idx < SrcLen; ++Idx)
+				{
+					Utf8[Idx] = static_cast<char8_t>(Src[Idx]);
+				}
+			}
+			else
+			{
+				UTF8CHAR* Utf8End = FPlatformString::Convert(reinterpret_cast<UTF8CHAR*>(Utf8), Utf8Len, Src, SrcLen);	
+				check((char8_t*)Utf8End - Utf8 == Utf8Len);
+			}
 		}
 	}
 
@@ -194,17 +205,29 @@ struct TStringBinding : ILeafRangeBinding
 		}
 		else if constexpr (sizeof(CharType) == sizeof(char8_t))
 		{
-			Dst.SetNum(SrcLen + 1);
+			Dst.SetNumUninitialized(SrcLen + 1);
 			FMemory::Memcpy(Dst.GetData(), Src, SrcLen);
 			Dst[SrcLen] = CharType('\0');	
 		}
 		else
 		{
 			int32 DstLen = FPlatformString::ConvertedLength<CharType>(Src, SrcLen);
-			Dst.SetNum(DstLen + 1);
-			CharType* DstEnd = FPlatformString::Convert(Dst.GetData(), DstLen, Src, SrcLen);
-			check(DstEnd - Dst.GetData() == DstLen);
-			*DstEnd = '\0';
+			Dst.SetNumUninitialized(DstLen + 1);
+			if (DstLen == SrcLen)
+			{
+				CharType* Out = Dst.GetData();
+				for (int32 Idx = 0; Idx < SrcLen; ++Idx)
+				{
+					Out[Idx] = static_cast<CharType>(Src[Idx]);
+				}
+				Out[SrcLen] = '\0';
+			}
+			else
+			{
+				CharType* DstEnd = FPlatformString::Convert(Dst.GetData(), DstLen, Src, SrcLen);
+				check(DstEnd - Dst.GetData() == DstLen);
+				*DstEnd = '\0';
+			}
 		}
 	}
 
diff --git a/Engine/Plugins/Experimental/PlainPropsUObject/Source/Private/PlainPropsUObjectBindings.cpp b/Engine/Plugins/Experimental/PlainPropsUObject/Source/Private/PlainPropsUObjectBindings.cpp
--- a/Engine/Plugins/Experimental/PlainPropsUObject/Source/Private/PlainPropsUObjectBindings.cpp
+++ b/Engine/Plugins/Experimental/PlainPropsUObject/Source/Private/PlainPropsUObjectBindings.cpp
@@ -3876,6 +3876,38 @@ static bool IncludeClass(UClass* Class)
 	return true;
 }
 
+FORCENOINLINE static void SavePlainProps(FBatchSaver& Batch, TConstArrayView<FInstance> Instances)
+{
+	for (const FInstance& Instance : Instances)
+	{	
+		Batch.Save(Instance.Id, Instance.Orig, Instance.Base);
+	}
+}
+
+FORCENOINLINE static void LoadPlainProps(FMemoryBatchLoader& Batch, TConstArrayView<FInstance> Instances)
+{
+	for (const FInstance& Instance : Instances)
+	{	
+		Batch.Load(Instance.PP);
+	}
+}
+
+FORCENOINLINE static void SaveTps(FPropertyWriter& Archive, TConstArrayView<FInstance> Instances)
+{
+	for (const FInstance& Instance : Instances)
+	{
+		Archive.WriteProperties(Instance.Orig, Instance.Base);
+	}
+}
+
+FORCENOINLINE static void LoadTps(FPropertyReader& Archive, TConstArrayView<FInstance> Instances)
+{
+	for (const FInstance& Instance : Instances)
+	{
+		Archive.ReadProperties(Instance.TPS, Instance.Base);
+	}
+}
+
 static int32 Roundtrip()
 {
 	UE_LOGFMT(LogPlainPropsUObject, Display, "Gathering all non-empty UObjects...");
@@ -3913,20 +3945,14 @@ static int32 Roundtrip()
 	TCustomPropertyBindings<FMemoryPropertyBatch> Customs(Plain.Properties, GUE.Customs);
 	{
 		FBatchSaver Batch(Customs.Overlay, GUObjectArray.GetObjectArrayNum());
-		for (const FInstance& Instance : Instances)
-		{	
-			Batch.Save(Instance.Id, Instance.Orig, Instance.Base);
-		}
+		SavePlainProps(Batch, Instances);
 		Plain.Data = Batch.Write(/* out */ Plain.RuntimeIds);
 	}
 
 	// Load
 	UE_LOGFMT(LogPlainPropsUObject, Display, "Loading UObjects from PlainProps...");
 	FMemoryBatchLoader Batch(Customs.Overlay, MakeMemoryView(Plain.Data), Plain.RuntimeIds);
-	for (const FInstance& Instance : Instances)
-	{
-		Batch.Load(Instance.PP);
-	}
+	LoadPlainProps(Batch, Instances);
 
 	// Diff original vs PlainProps
 	UE_LOGFMT(LogPlainPropsUObject, Display, "Diffing UObjects roundtripped via PlainProps...");
@@ -3950,19 +3976,13 @@ static int32 Roundtrip()
 	FArchivedProperties Tps;
 	{
 		FPropertyWriter Archive(/* out */ Tps);
-		for (const FInstance& Instance : Instances)
-		{
-			Archive.WriteProperties(Instance.Orig, Instance.Base);
-		}
+		SaveTps(Archive, Instances);
 	}
 
 	UE_LOGFMT(LogPlainPropsUObject, Display, "Loading UObjects from TPS archive...");
 	{
 		FPropertyReader Archive(Tps);
-		for (const FInstance& Instance : Instances)
-		{
-			Archive.ReadProperties(Instance.TPS, Instance.Base);
-		}
+		LoadTps(Archive, Instances);
 	}
 
 	// Diff original vs TPS
