PlainProps: Move NumXYZ interface to FIds base class and safe up the FDebugIds print functions.

#rnx

diff --git a/Engine/Plugins/Experimental/PlainProps/Source/Private/PlainPropsInternalText.h b/Engine/Plugins/Experimental/PlainProps/Source/Private/PlainPropsInternalText.h
--- a/Engine/Plugins/Experimental/PlainProps/Source/Private/PlainPropsInternalText.h
+++ b/Engine/Plugins/Experimental/PlainProps/Source/Private/PlainPropsInternalText.h
@@ -26,6 +26,8 @@ struct FLiterals
 	FAnsiStringView Super = "__super__";
 	FAnsiStringView Dynamic = "__dynamic__";
 
+	FAnsiStringView Oob = "ERR_oob";
+
 	FAnsiStringView Ranges[9] = { 
 		("Uni"), ("i8"), ("u8"), ("i16"), ("u16"), ("i32"), ("u32"), ("i64"), ("u64") };
 
@@ -37,7 +39,7 @@ struct FLiterals
 		{("hex8"),		("hex16"),		("hex32"),		("hex64")},
 		{("enum8"),		("enum16"),		("enum32"),		("enum64")},
 		{("utf8"),		("utf16"),		("utf32"),		("ERR_utf64")},
-		{("ERR_oob"),	("ERR_oob"),	("ERR_oob"),	("ERR_oob")}};
+		{Oob,			Oob,			Oob,			Oob}};
 
 	FAnsiStringView Widths[4] = { ("u8"), ("u16"), ("u32"), ("u64") };
 };
diff --git a/Engine/Plugins/Experimental/PlainProps/Source/Private/PlainPropsPrint.cpp b/Engine/Plugins/Experimental/PlainProps/Source/Private/PlainPropsPrint.cpp
--- a/Engine/Plugins/Experimental/PlainProps/Source/Private/PlainPropsPrint.cpp
+++ b/Engine/Plugins/Experimental/PlainProps/Source/Private/PlainPropsPrint.cpp
@@ -1142,35 +1142,72 @@ void FIds::AppendString(FUtf8Builder& Out, FStructSchemaId Name) const
 FString FDebugIds::Print(FNameId Name) const
 {
 	TUtf8StringBuilder<128> Out;
-	Ids.AppendString(Out, Name);
+	if (Name.Idx < Ids.NumNames())
+	{
+		Ids.AppendString(Out, Name);
+	}
+	else
+	{
+		Out << GLiterals.Oob;
+	}
 	return FString(FStringView(StringCast<TCHAR>(Out.GetData(), Out.Len())));
 }
 
 FString FDebugIds::Print(FMemberId Name) const
 {
 	TUtf8StringBuilder<128> Out;
-	Ids.AppendString(Out, Name.Id);
+	if (Name.Id.Idx < Ids.NumNames())
+	{
+		Ids.AppendString(Out, Name.Id);
+	}
+	else
+	{
+		Out << GLiterals.Oob;
+	}
 	return FString(FStringView(StringCast<TCHAR>(Out.GetData(), Out.Len())));
 }
 
 FString FDebugIds::Print(FOptionalMemberId Name) const
 {
 	TUtf8StringBuilder<128> Out;
-	Ids.AppendString(Out, Name);
+	if (!Name || Name.Get().Id.Idx < Ids.NumNames())
+	{
+		Ids.AppendString(Out, Name);
+	}
+	else
+	{
+		Out << GLiterals.Oob;
+	}
 	return FString(FStringView(StringCast<TCHAR>(Out.GetData(), Out.Len())));
 }
 
 FString FDebugIds::Print(FScopeId Scope) const
 {
 	TUtf8StringBuilder<128> Out;
-	Ids.AppendString(Out, Scope);
+	if (Scope.IsFlat() && Scope.AsFlat().Name.Idx < Ids.NumNames() ||
+		Scope.AsNested().Idx < Ids.NumNestedScopes())
+	{
+		Ids.AppendString(Out, Scope);
+	}
+	else
+	{
+		Out << GLiterals.Oob;
+	}
 	return FString(FStringView(StringCast<TCHAR>(Out.GetData(), Out.Len())));
 }
 
 FString FDebugIds::Print(FTypenameId Typename) const
 {
 	TUtf8StringBuilder<128> Out;
-	Ids.AppendString(Out, Typename);
+	if (Typename.IsConcrete() && Typename.AsConcrete().Id.Idx < Ids.NumNames() ||
+		Typename.AsParametric().Idx < Ids.NumNames())
+	{
+		Ids.AppendString(Out, Typename);
+	}
+	else
+	{
+		Out << GLiterals.Oob;
+	}
 	return FString(FStringView(StringCast<TCHAR>(Out.GetData(), Out.Len())));
 }
 
@@ -1179,6 +1216,11 @@ FString FDebugIds::Print(FConcreteTypenameId Typename) const
 	return Print(FTypenameId(Typename));
 }
 
+FString FDebugIds::Print(FParametricTypeId Typename) const
+{
+	return Print(FTypenameId(Typename));
+}
+
 FString FDebugIds::Print(FTypeId Type) const
 {
 	TUtf8StringBuilder<128> Out;
@@ -1189,14 +1231,28 @@ FString FDebugIds::Print(FTypeId Type) const
 FString FDebugIds::Print(FEnumSchemaId Name) const
 {
 	TUtf8StringBuilder<128> Out;
-	Ids.AppendString(Out, Name);
+	if (Name.Idx < Ids.NumEnums())
+	{
+		Ids.AppendString(Out, Name);
+	}
+	else
+	{
+		Out << GLiterals.Oob;
+	}
 	return FString(FStringView(StringCast<TCHAR>(Out.GetData(), Out.Len())));
 }
 
 FString FDebugIds::Print(FStructSchemaId Name) const
 {
 	TUtf8StringBuilder<128> Out;
-	Ids.AppendString(Out, Name);
+	if (Name.Idx < Ids.NumStructs())
+	{
+		Ids.AppendString(Out, Name);
+	}
+	else
+	{
+		Out << GLiterals.Oob;
+	}
 	return FString(FStringView(StringCast<TCHAR>(Out.GetData(), Out.Len())));
 }
 
diff --git a/Engine/Plugins/Experimental/PlainProps/Source/Public/PlainPropsIndex.h b/Engine/Plugins/Experimental/PlainProps/Source/Public/PlainPropsIndex.h
--- a/Engine/Plugins/Experimental/PlainProps/Source/Public/PlainPropsIndex.h
+++ b/Engine/Plugins/Experimental/PlainProps/Source/Public/PlainPropsIndex.h
@@ -60,6 +60,11 @@ public:
 	PLAINPROPS_API FEnumSchemaId		IndexEnum(FTypeId Type);
 	PLAINPROPS_API FStructSchemaId		IndexStruct(FTypeId Type);
 	
+	virtual uint32						NumNestedScopes() const override final				{ return IntCastChecked<uint32>(NestedScopes.Num()); }
+	virtual uint32						NumParametricTypes() const override final			{ return IntCastChecked<uint32>(ParametricTypes.Num()); }
+	virtual uint32						NumEnums() const override final						{ return Enums.Num(); }
+	virtual uint32						NumStructs() const override final					{ return Structs.Num(); }
+	
 	virtual FNestedScope				Resolve(FNestedScopeId Id) const override final		{ return NestedScopes.Resolve(Id); }
 	virtual FParametricTypeView			Resolve(FParametricTypeId Id) const override final	{ return ParametricTypes.Resolve(Id); }
 	PLAINPROPS_API virtual FTypeId		Resolve(FEnumSchemaId Id) const override final;
@@ -68,10 +73,6 @@ public:
 	const FNestedScopeIndexer&			GetNestedScopes() const								{ return NestedScopes; }
 	const FParametricTypeIndexer&		GetParametricTypes() const							{ return ParametricTypes; }
 
-	virtual uint32						NumNames() const = 0;
-	uint32								NumEnums() const									{ return Enums.Num(); }
-	uint32								NumStructs() const									{ return Structs.Num(); }
-	
 	FConcreteTypenameId					IndexRangeBindName(const char* Name)				{ return InitParameterName(Name); }
 protected:
 	FNestedScopeIndexer					NestedScopes;
diff --git a/Engine/Plugins/Experimental/PlainProps/Source/Public/PlainPropsTypes.h b/Engine/Plugins/Experimental/PlainProps/Source/Public/PlainPropsTypes.h
--- a/Engine/Plugins/Experimental/PlainProps/Source/Public/PlainPropsTypes.h
+++ b/Engine/Plugins/Experimental/PlainProps/Source/Public/PlainPropsTypes.h
@@ -425,8 +425,14 @@ class FIds
 public:
 	virtual ~FIds() {}
 
-	virtual FParametricTypeView			Resolve(FParametricTypeId Id) const = 0;
+	virtual uint32						NumNames() const = 0;
+	virtual uint32						NumNestedScopes() const = 0;
+	virtual uint32						NumParametricTypes() const = 0;
+	virtual uint32						NumEnums() const = 0;
+	virtual uint32						NumStructs() const = 0;
+
 	virtual FNestedScope				Resolve(FNestedScopeId Id) const = 0;
+	virtual FParametricTypeView			Resolve(FParametricTypeId Id) const = 0;
 	virtual FTypeId						Resolve(FEnumSchemaId Id) const = 0;
 	virtual FTypeId						Resolve(FStructSchemaId Id) const = 0;
 
@@ -452,6 +458,7 @@ public:
 	PLAINPROPS_API FString				Print(FScopeId Scope) const;
 	PLAINPROPS_API FString				Print(FTypenameId Typename) const;
 	PLAINPROPS_API FString				Print(FConcreteTypenameId Typename) const;
+	PLAINPROPS_API FString				Print(FParametricTypeId Typename) const;
 	PLAINPROPS_API FString				Print(FTypeId Type) const;
 	PLAINPROPS_API FString				Print(FEnumSchemaId Enum) const;
 	PLAINPROPS_API FString				Print(FStructSchemaId Struct) const;
diff --git a/Engine/Plugins/Experimental/PlainProps/Source/Tests/PlainPropsTests.cpp b/Engine/Plugins/Experimental/PlainProps/Source/Tests/PlainPropsTests.cpp
--- a/Engine/Plugins/Experimental/PlainProps/Source/Tests/PlainPropsTests.cpp
+++ b/Engine/Plugins/Experimental/PlainProps/Source/Tests/PlainPropsTests.cpp
@@ -434,6 +434,7 @@ public:
 		check(NextStr == Data.GetDataEnd()); // end with null-terminator
 	}
 	
+	uint32 NumNames() const { return IntCastChecked<uint32>(Names.Num()); }
 	FAnsiStringView operator[](FNameId Id) const { return MakeStringView(Names[Id.Idx]); }
 	FAnsiStringView operator[](FMemberId Name) const { return MakeStringView(Names[Name.Id.Idx]); }
 	FAnsiStringView operator[](FOptionalMemberId Name) const { return Name ? operator[](Name.Get()) : "Super"; }
@@ -449,6 +450,12 @@ class FTestBatchNames : public FIds
 public:
 	FTestBatchNames(const FTestNameReader& InNames, const FSchemaBatch& InSchemas) : Names(InNames), Schemas(InSchemas) {}
 
+	virtual uint32					NumNames() const override final						{ return Names.NumNames(); }
+	virtual uint32					NumNestedScopes() const override final				{ return Schemas.NumNestedScopes; }
+	virtual uint32					NumParametricTypes() const override final			{ return Schemas.NumParametricTypes; }
+	virtual uint32					NumEnums() const override final						{ return Schemas.NumSchemas - Schemas.NumStructSchemas; }
+	virtual uint32					NumStructs() const override final					{ return Schemas.NumStructSchemas; }
+
 	virtual FNestedScope			Resolve(FNestedScopeId Id) const override final		{ return ResolveNestedScope(Schemas, Id); }
 	virtual FParametricTypeView		Resolve(FParametricTypeId Id) const override final	{ return ResolveParametricType(Schemas, Id); }
 	virtual FTypeId					Resolve(FEnumSchemaId Id) const override final		{ return ResolveEnumSchema(Schemas, Id).Type; }
diff --git a/Engine/Plugins/Experimental/PlainProps/Source/Tests/PlainPropsUeCoreTests.cpp b/Engine/Plugins/Experimental/PlainProps/Source/Tests/PlainPropsUeCoreTests.cpp
--- a/Engine/Plugins/Experimental/PlainProps/Source/Tests/PlainPropsUeCoreTests.cpp
+++ b/Engine/Plugins/Experimental/PlainProps/Source/Tests/PlainPropsUeCoreTests.cpp
@@ -269,6 +269,9 @@ class FSchemaBatchIds : public FIds
 public:
 	FSchemaBatchIds(const FSchemaBatch& InSchemas) : Schemas(InSchemas) {}
 
+	virtual uint32				NumEnums() const override final						{ return Schemas.NumSchemas - Schemas.NumStructSchemas; }
+	virtual uint32				NumStructs() const override final					{ return Schemas.NumStructSchemas; }
+
 	virtual FTypeId				Resolve(FEnumSchemaId Id) const override final		{ return ResolveEnumSchema(Schemas, Id).Type; }
 	virtual FTypeId				Resolve(FStructSchemaId Id) const override final	{ return ResolveStructSchema(Schemas, Id).Type; }
 
@@ -282,9 +285,13 @@ public:
 	using FSchemaBatchIds::FSchemaBatchIds;
 	using FSchemaBatchIds::Resolve;
 
+	virtual uint32				NumNames() const override final						{ return GNames.NumNames(); }
+	virtual uint32				NumNestedScopes() const override final				{ return GNames.NumNestedScopes(); }
+	virtual uint32				NumParametricTypes() const override final			{ return GNames.NumParametricTypes(); }
+
 	virtual FNestedScope		Resolve(FNestedScopeId Id) const override			{ return GNames.Resolve(Id); }
 	virtual FParametricTypeView	Resolve(FParametricTypeId Id) const override		{ return GNames.Resolve(Id); }
-	virtual void				AppendString(FUtf8Builder& Out, FNameId Name) const override { GNames.ResolveName(Name).AppendString(Out); }
+	virtual void				AppendString(FUtf8Builder& Out, FNameId Name) const override final { GNames.ResolveName(Name).AppendString(Out); }
 };
 
 class FStableBatchIds final : public FSchemaBatchIds
@@ -293,9 +300,13 @@ public:
 	using FSchemaBatchIds::FSchemaBatchIds;
 	using FSchemaBatchIds::Resolve;
 
+	virtual uint32				NumNames() const override final						{ return IntCastChecked<uint32>(Names.Num()); }
+	virtual uint32				NumNestedScopes() const override final				{ return Schemas.NumNestedScopes; }
+	virtual uint32				NumParametricTypes() const override final			{ return Schemas.NumParametricTypes; }
+
 	virtual FNestedScope		Resolve(FNestedScopeId Id) const override			{ return ResolveNestedScope(Schemas, Id); }
-	virtual FParametricTypeView	Resolve(FParametricTypeId Id) const override 		{ return ResolveParametricType(Schemas, Id); }
-	virtual void				AppendString(FUtf8Builder& Out, FNameId Name) const override { Names[Name.Idx].AppendString(Out); }
+	virtual FParametricTypeView	Resolve(FParametricTypeId Id) const override		{ return ResolveParametricType(Schemas, Id); }
+	virtual void				AppendString(FUtf8Builder& Out, FNameId Name) const override final { Names[Name.Idx].AppendString(Out); }
 
 	TConstArrayView<FName> Names;
 };
@@ -305,6 +316,12 @@ class FLoadBatchNames : public FIds
 public:
 	FLoadBatchNames(const FIdBinding& InBinding, const FSchemaBatch& InSchemas) : Ids(GNames), Binding(InBinding), Schemas(InSchemas) {}
 
+	virtual uint32				NumNames() const override final						{ return Ids.NumNames(); }
+	virtual uint32				NumNestedScopes() const override final				{ return IntCastChecked<uint32>(Binding.NestedScopes.Num()); }
+	virtual uint32				NumParametricTypes() const override final			{ return IntCastChecked<uint32>(Binding.ParametricTypes.Num()); }
+	virtual uint32				NumEnums() const override final						{ return Schemas.NumSchemas - Schemas.NumStructSchemas; }
+	virtual uint32				NumStructs() const override final					{ return Schemas.NumStructSchemas; }
+
 	virtual FNestedScope		Resolve(FNestedScopeId Id) const override final		{ return Ids.Resolve(Binding.Remap(Id)); }
 	virtual FParametricTypeView	Resolve(FParametricTypeId Id) const override final	{ return Ids.Resolve(Binding.Remap(Id)); }
 	virtual FTypeId				Resolve(FEnumSchemaId Id) const override final		{ return ResolveEnumSchema(Schemas, Id).Type; }
