PlainProps: Compile fix in PlainPropsPrint.cpp for cl 39338567

#rnx

diff --git a/Engine/Plugins/Experimental/PlainProps/Source/Private/PlainPropsPrint.cpp b/Engine/Plugins/Experimental/PlainProps/Source/Private/PlainPropsPrint.cpp
--- a/Engine/Plugins/Experimental/PlainProps/Source/Private/PlainPropsPrint.cpp
+++ b/Engine/Plugins/Experimental/PlainProps/Source/Private/PlainPropsPrint.cpp
@@ -1184,7 +1184,7 @@ FString FDebugIds::Print(FOptionalMemberId Name) const
 FString FDebugIds::Print(FScopeId Scope) const
 {
 	TUtf8StringBuilder<128> Out;
-	if (Scope.IsFlat() && Scope.AsFlat().Name.Idx < Ids.NumNames() ||
+	if ((Scope.IsFlat() && Scope.AsFlat().Name.Idx < Ids.NumNames()) ||
 		Scope.AsNested().Idx < Ids.NumNestedScopes())
 	{
 		Ids.AppendString(Out, Scope);
@@ -1199,7 +1199,7 @@ FString FDebugIds::Print(FScopeId Scope) const
 FString FDebugIds::Print(FTypenameId Typename) const
 {
 	TUtf8StringBuilder<128> Out;
-	if (Typename.IsConcrete() && Typename.AsConcrete().Id.Idx < Ids.NumNames() ||
+	if ((Typename.IsConcrete() && Typename.AsConcrete().Id.Idx < Ids.NumNames()) ||
 		Typename.AsParametric().Idx < Ids.NumNames())
 	{
 		Ids.AppendString(Out, Typename);
