PlainProps split runtime and serialized struct and ids statically

#rnx

diff --git a/Engine/Plugins/Experimental/PlainProps/Source/Private/PlainPropsBind.cpp b/Engine/Plugins/Experimental/PlainProps/Source/Private/PlainPropsBind.cpp
--- a/Engine/Plugins/Experimental/PlainProps/Source/Private/PlainPropsBind.cpp
+++ b/Engine/Plugins/Experimental/PlainProps/Source/Private/PlainPropsBind.cpp
@@ -25,7 +25,7 @@ uint32 FSchemaBinding::CalculateSize() const
 {
 	uint32 Out = sizeof(FSchemaBinding) + (NumMembers + NumInnerRanges) * sizeof(FMemberBindType);
 	Out = Align(Out + NumMembers * sizeof(uint32), sizeof(uint32));
-	Out = Align(Out + NumInnerSchemas * sizeof(FSchemaId), sizeof(FSchemaId));
+	Out = Align(Out + NumInnerSchemas * sizeof(FInnerId), sizeof(FInnerId));
 	Out = Align(Out + NumInnerRanges * sizeof(FRangeBinding), sizeof(FRangeBinding));
 	return Out;
 }
@@ -62,7 +62,7 @@ uint64 FMemberVisitor::GrabMemberOffset()
 FLeafMemberBinding FMemberVisitor::GrabLeaf()
 {
 	FUnpackedLeafBindType Leaf = PeekType().AsLeaf();
-	FOptionalEnumSchemaId Enum = Leaf.Type == ELeafBindType::Enum ? ToOptional(GrabEnumSchema()) : NoId;
+	FOptionalEnumId Enum = Leaf.Type == ELeafBindType::Enum ? ToOptional(GrabEnumSchema()) : NoId;
 	uint64 Offset = GrabMemberOffset();
 
 	return {Leaf, Enum, Offset};
@@ -71,7 +71,7 @@ FLeafMemberBinding FMemberVisitor::GrabLeaf()
 FStructMemberBinding FMemberVisitor::GrabStruct()
 {
 	checkf(!PeekType().AsStruct().IsDynamic, TEXT("Bound structs can't be dynamic"));
-	return { PeekType().AsStruct(), static_cast<FStructSchemaId>(GrabInnerSchema()), GrabMemberOffset() };
+	return { PeekType().AsStruct(), GrabInnerSchema().AsStruct(), GrabMemberOffset() };
 }
 
 static bool HasSchema(FMemberBindType Type)
@@ -93,7 +93,7 @@ FRangeMemberBinding FMemberVisitor::GrabRange()
 	const FRangeBinding* RangeBindings = Schema.GetRangeBindings() + InnerRangeIdx;
 	check(MaxSize == RangeBindings[0].GetSizeType());
 	FMemberBindTypeRange InnerTypes = GrabInnerTypes();
-	FOptionalSchemaId InnermostSchema = HasSchema(InnerTypes.Last()) ? ToOptional(GrabInnerSchema()) : NoId; 
+	FOptionalInnerId InnermostSchema = HasSchema(InnerTypes.Last()) ? ToOptional(GrabInnerSchema()) : NoId; 
 	uint64 Offset = GrabMemberOffset();
 		
 	return { &InnerTypes[0], RangeBindings, static_cast<uint16>(InnerTypes.Num()), InnermostSchema, Offset};
@@ -115,16 +115,16 @@ void FMemberVisitor::SkipMember()
 	++MemberIdx;
 }
 
-FStructSchemaId FMemberVisitor::GrabSuper()
+FStructId FMemberVisitor::GrabSuper()
 {
 	check(MemberIdx == 0);
 	checkSlow(Schema.Members[0].AsByte() == SuperStructType.AsByte());
 	MemberIdx = 1;
 	InnerSchemaIdx = 1;
-	return static_cast<FStructSchemaId>(Schema.GetInnerSchemas()[0]);
+	return Schema.GetInnerSchemas()[0].AsStruct();
 }
 
-FSchemaId FMemberVisitor::GrabInnerSchema()
+FInnerId FMemberVisitor::GrabInnerSchema()
 {
 	check(InnerSchemaIdx < Schema.NumInnerSchemas);
 	return Schema.GetInnerSchemas()[InnerSchemaIdx++];
@@ -197,7 +197,7 @@ struct FCustomBindingHandle
 	FInnersHandle*					GetLoweredInners() { return reinterpret_cast<FInnersHandle*>(Bindings + 2 * Max); }
 };
 
-void FCustomBindingMap::Bind(FStructSchemaId BindId, ICustomBinding& Binding, const FStructDeclaration& Declaration, FInnersHandle LoweredInners)
+void FCustomBindingMap::Bind(FStructId BindId, ICustomBinding& Binding, const FStructDeclaration& Declaration, FInnersHandle LoweredInners)
 {
 	// Add key
 	bool bExists = false;
@@ -233,7 +233,7 @@ void FCustomBindingMap::Bind(FStructSchemaId BindId, ICustomBinding& Binding, co
 	UpdateCulling(/* in-out */ Window, BindId.Idx);
 }
 
-inline FCustomBindingHandle FCustomBindingMap::Find(FStructSchemaId BindId) const
+inline FCustomBindingHandle FCustomBindingMap::Find(FStructId BindId) const
 {
 	if (Cull(Window, BindId.Idx))
 	{
@@ -245,7 +245,7 @@ inline FCustomBindingHandle FCustomBindingMap::Find(FStructSchemaId BindId) cons
 	return {};
 }
 
-void FCustomBindingMap::Drop(FStructSchemaId BindId)
+void FCustomBindingMap::Drop(FStructId BindId)
 {
 	int32 NumRemoved = Keys.Remove(BindId);
 	checkf(NumRemoved == 1, TEXT("'%s' unbound"), *Debug.Print(BindId));
@@ -280,20 +280,20 @@ static bool operator==(FInnerStruct A, FInnerStruct B) { return A.Name == B.Name
 
 ////////////////////////////////////////////////////////////////////////////////////////////////
 
-void FCustomBindings::BindStruct(FStructSchemaId BindId, ICustomBinding& Binding, const FStructDeclaration& Declaration, TConstArrayView<FInnerStruct> LoweredInners)
+void FCustomBindings::BindStruct(FStructId BindId, ICustomBinding& Binding, const FStructDeclaration& Declaration, TConstArrayView<FInnerStruct> LoweredInners)
 {
 	FInnersHandle StoredInners = StoreInners(/* out */ BottomInners, LoweredInners);
 	check(Algo::Compare(LoweredInners, FetchInners(BottomInners, StoredInners)));
 	Map.Bind(BindId, Binding, Declaration, StoredInners);
 }
 
-const ICustomBinding* FCustomBindings::FindStruct(FStructSchemaId BindId) const
+const ICustomBinding* FCustomBindings::FindStruct(FStructId BindId) const
 {
 	FCustomBindingHandle Handle = Find(BindId);
 	return Handle ? *Handle.Bindings : nullptr;
 }
 
-const ICustomBinding* FCustomBindings::FindStruct(FStructSchemaId BindId, TConstArrayView<FInnerStruct>& OutLoweredInners) const
+const ICustomBinding* FCustomBindings::FindStruct(FStructId BindId, TConstArrayView<FInnerStruct>& OutLoweredInners) const
 {
 	if (FCustomBindingHandle Handle = Find(BindId))
 	{
@@ -303,30 +303,30 @@ const ICustomBinding* FCustomBindings::FindStruct(FStructSchemaId BindId, TConst
 	return nullptr;
 }
 
-ICustomBinding* FCustomBindings::FindStructToSave(FStructSchemaId BindId, const FStructDeclaration*& OutDeclaration) const
+ICustomBinding* FCustomBindings::FindStructToSave(FStructId BindId, const FStructDeclaration*& OutDeclaration) const
 {
 	FCustomBindingHandle Handle = Find(BindId);
 	OutDeclaration = Handle ? *Handle.GetDeclarations() : nullptr;
 	return Handle ? *Handle.Bindings : nullptr;
 }
 
-const FStructDeclaration* FCustomBindings::FindDeclaration(FStructSchemaId BindId) const
+const FStructDeclaration* FCustomBindings::FindDeclaration(FStructId BindId) const
 { 
 	FCustomBindingHandle Handle = Find(BindId);
 	return Handle ? *Handle.GetDeclarations() : nullptr;
 }
 
-void FCustomBindings::DropStruct(FStructSchemaId BindId)
+void FCustomBindings::DropStruct(FStructId BindId)
 {
 	Map.Drop(BindId);
 }
 
-FCustomBindingHandle FCustomBindingsBottom::Find(FStructSchemaId BindId) const
+FCustomBindingHandle FCustomBindingsBottom::Find(FStructId BindId) const
 {
 	return Map.Find(BindId);
 }
 
-FCustomBindingHandle FCustomBindingsOverlay::Find(FStructSchemaId BindId) const
+FCustomBindingHandle FCustomBindingsOverlay::Find(FStructId BindId) const
 {
 	FCustomBindingHandle Handle = Map.Find(BindId);
 	return Handle ? Handle : Underlay.Find(BindId);
@@ -334,23 +334,23 @@ FCustomBindingHandle FCustomBindingsOverlay::Find(FStructSchemaId BindId) const
 
 ////////////////////////////////////////////////////////////////////////////////////////////////
 
-inline FType IndexInnermost(FIdIndexerBase& Ids, FMemberBindType Type, FOptionalSchemaId Schema)
+inline FType IndexInnermost(FIdIndexerBase& Ids, FMemberBindType Type, FOptionalInnerId Schema)
 {
 	if (Type.IsStruct())
 	{
-		return Ids.Resolve(static_cast<FStructSchemaId>(Schema.Get()));
+		return Ids.Resolve(Schema.Get().AsStruct());
 	}
 
 	FUnpackedLeafType Leaf = ToLeafType(Type.AsLeaf());
 	if (Leaf.Type == ELeafType::Enum)
 	{
-		return Ids.Resolve(static_cast<FEnumSchemaId>(Schema.Get()));
+		return Ids.Resolve(Schema.Get().AsEnum());
 	}
 
 	return Ids.MakeLeafParameter(Leaf);
 }
 
-FDualType	FMemberBinding::IndexParameterName(FIdIndexerBase& Ids) const
+FDualType FMemberBinding::IndexParameterName(FIdIndexerBase& Ids) const
 {
 	FType BindType = IndexInnermost(Ids, InnermostType, InnermostSchema);
 	FType DeclType = BindType;
@@ -388,7 +388,27 @@ static uint16 CountRanges(TConstArrayView<FMemberBinding> Members)
 	return IntCastChecked<uint16>(Out);
 }
 
-void FSchemaBindings::BindStruct(FStructSchemaId BindId, FStructSchemaId DeclId, TConstArrayView<FMemberBinding> Members)
+struct FMemberBinder : FMemberBinderBase
+{
+	FMemberBinder(FSchemaBinding& InSchema)
+	: FMemberBinderBase(InSchema)
+	, InnerSchemaIt(const_cast<FInnerId*>(InSchema.GetInnerSchemas()))
+	{}
+	
+	~FMemberBinder()
+	{
+		check(Align(InnerSchemaIt, alignof(FRangeBinding)) == (const void*)Schema.GetRangeBindings() || Schema.NumInnerRanges == 0);
+	}
+
+	void AddInnerSchema(FInnerId InnermostSchema)
+	{
+		*InnerSchemaIt++ = InnermostSchema;
+	}
+
+	FInnerId* InnerSchemaIt;
+};
+
+void FSchemaBindings::BindStruct(FStructId BindId, FStructId DeclId, TConstArrayView<FMemberBinding> Members)
 {
 	// Make header, allocate and copy header
 	FSchemaBinding Header = { DeclId, IntCastChecked<uint16>(Members.Num()), CountInnerSchemas(Members), CountRanges(Members) };
@@ -423,31 +443,31 @@ void FSchemaBindings::BindStruct(FStructSchemaId BindId, FStructSchemaId DeclId,
 	Bindings[BindId.Idx].Reset(Schema);
 }
 
-const FSchemaBinding* FSchemaBindings::FindStruct(FStructSchemaId BindId) const
+const FSchemaBinding* FSchemaBindings::FindStruct(FStructId BindId) const
 {
 	return BindId.Idx < (uint32)Bindings.Num() ? Bindings[BindId.Idx].Get() : nullptr;
 }
 
-const FSchemaBinding& FSchemaBindings::GetStruct(FStructSchemaId BindId) const
+const FSchemaBinding& FSchemaBindings::GetStruct(FStructId BindId) const
 {
 	checkf(BindId.Idx < (uint32)Bindings.Num() && Bindings[BindId.Idx], TEXT("'%s' is unbound"), *Debug.Print(BindId));
 	return *Bindings[BindId.Idx].Get();
 }
 
-void FSchemaBindings::DropStruct(FStructSchemaId BindId)
+void FSchemaBindings::DropStruct(FStructId BindId)
 {
 	checkf(BindId.Idx < (uint32)Bindings.Num() && Bindings[BindId.Idx], TEXT("'%s' is unbound"), *Debug.Print(BindId));
 	Bindings[BindId.Idx].Reset();
 }
 
-FStructSchemaId FSchemaBindings::GetDeclId(FStructSchemaId BindId) const
+FStructId FSchemaBindings::GetDeclId(FStructId BindId) const
 {
 	return GetStruct(BindId).DeclId;
 }
 
 //////////////////////////////////////////////////////////////////////////
 
-FStructSchemaId FStructBindIds::GetDeclId(FStructSchemaId BindId) const
+FStructId FStructBindIds::GetDeclId(FStructId BindId) const
 {
 	const FStructDeclaration* CustomDecl = Customs.FindDeclaration(BindId);
 	return CustomDecl ? CustomDecl->Id : Schemas.GetDeclId(BindId);
@@ -455,12 +475,12 @@ FStructSchemaId FStructBindIds::GetDeclId(FStructSchemaId BindId) const
 
 //////////////////////////////////////////////////////////////////////////
 
-TArray<FStructSchemaId> IndexRuntimeIds(const FSchemaBatch& Schemas, FIdIndexerBase& Indexer)
+TArray<FStructId> IndexRuntimeIds(const FSchemaBatch& Schemas, FIdIndexerBase& Indexer)
 {
 	const uint8* Base = reinterpret_cast<const uint8*>(&Schemas);
 	const uint32* Offsets = Schemas.SchemaOffsets;
 
-	TArray<FStructSchemaId> Out;
+	TArray<FStructId> Out;
 	Out.SetNumUninitialized(Schemas.NumStructSchemas);
 	for (int32 Idx = 0; Idx < Out.Num(); ++Idx)
 	{
@@ -475,7 +495,7 @@ uint32 FIdTranslatorBase::CalculateTranslationSize(int32 NumSavedNames, const FS
 {
 	static_assert(sizeof(FNameId) == sizeof(FNestedScopeId));
 	static_assert(sizeof(FNameId) == sizeof(FParametricTypeId));
-	static_assert(sizeof(FNameId) == sizeof(FSchemaId));
+	static_assert(sizeof(FNameId) == sizeof(FInnerId));
 	return sizeof(FNameId) * (NumSavedNames + Batch.NumNestedScopes + Batch.NumParametricTypes + Batch.NumSchemas);
 }
 
@@ -511,20 +531,20 @@ static void TranslateParametricTypeIds(TArrayView<FParametricTypeId> Out, FIdInd
 	}
 }
 
-static void TranslateSchemaIds(TArrayView<FSchemaId> Out, FIdIndexerBase& Indexer, const FIdBinding& To, const FSchemaBatch& From)
+static void TranslateSchemaIds(TArrayView<FInnerId> Out, FIdIndexerBase& Indexer, const FIdBinding& To, const FSchemaBatch& From)
 {
 	uint32 OutIdx = 0;
 	for (const FStructSchema& FromSchema : GetStructSchemas(From))
 	{
 		FType ToType = To.Remap(FromSchema.Type);
 		checkSlow(ToType.Name.NumParameters == FromSchema.Type.Name.NumParameters);
-		Out[OutIdx++] = Indexer.IndexStruct(ToType);
+		Out[OutIdx++] = FInnerId(Indexer.IndexStruct(ToType));
 	}
 	
 	for (const FEnumSchema& FromSchema : GetEnumSchemas(From))
 	{
 		FType ToType = To.Remap(FromSchema.Type);
-		Out[OutIdx++] = Indexer.IndexEnum(ToType);
+		Out[OutIdx++] = FInnerId(Indexer.IndexEnum(ToType));
 	}
 }
 
@@ -532,7 +552,7 @@ FIdBinding FIdTranslatorBase::TranslateIds(FMutableMemoryView To, FIdIndexerBase
 {
 	TArrayView<FNestedScopeId> ToScopes(static_cast<FNestedScopeId*>(To.GetData()), From.NumNestedScopes);
 	TArrayView<FParametricTypeId> ToParametricTypes(reinterpret_cast<FParametricTypeId*>(ToScopes.end()), From.NumParametricTypes);
-	TArrayView<FSchemaId> ToSchemas(reinterpret_cast<FSchemaId*>(ToParametricTypes.end()), From.NumSchemas);
+	TArrayView<FInnerId> ToSchemas(reinterpret_cast<FInnerId*>(ToParametricTypes.end()), From.NumSchemas);
 	FIdBinding Out = {ToNames, ToScopes, ToParametricTypes, ToSchemas};
 	check(uintptr_t(To.GetDataEnd()) == uintptr_t(ToSchemas.end()));
 
diff --git a/Engine/Plugins/Experimental/PlainProps/Source/Private/PlainPropsBuild.cpp b/Engine/Plugins/Experimental/PlainProps/Source/Private/PlainPropsBuild.cpp
--- a/Engine/Plugins/Experimental/PlainProps/Source/Private/PlainPropsBuild.cpp
+++ b/Engine/Plugins/Experimental/PlainProps/Source/Private/PlainPropsBuild.cpp
@@ -118,7 +118,7 @@ FBuiltRange* CloneLeaves(FScratchAllocator& Scratch, uint64 Num, const void* InD
 
 //////////////////////////////////////////////////////////////////////////
 
-void FMemberBuilder::AddSuperStruct(FStructSchemaId SuperSchema, FBuiltStruct* SuperStruct)
+void FMemberBuilder::AddSuperStruct(FStructId SuperSchema, FBuiltStruct* SuperStruct)
 {
 	check(SuperStruct);
 	check(Members.IsEmpty());
@@ -222,13 +222,13 @@ uint64 ValueCast(double Value)
 
 //////////////////////////////////////////////////////////////////////////
 
-template<typename SchemaIdType>
-static FMemberSchema MakeMemberSchema(FMemberType Type, SchemaIdType InnerSchema)
+template<typename InnerIdType>
+static FMemberSchema MakeMemberSchema(FMemberType Type, InnerIdType InnerSchema)
 {
-	return {Type, Type, 0, FOptionalSchemaId(InnerSchema), nullptr};
+	return {Type, Type, 0, FOptionalInnerId(InnerSchema), nullptr};
 }
 
-FBuiltMember::FBuiltMember(FMemberId Name, FUnpackedLeafType Leaf, FOptionalEnumSchemaId Enum, uint64 Value)
+FBuiltMember::FBuiltMember(FMemberId Name, FUnpackedLeafType Leaf, FOptionalEnumId Enum, uint64 Value)
 : FBuiltMember(Name, MakeMemberSchema(Leaf.Pack(), Enum), { .Leaf = Value})
 {}
 
@@ -236,13 +236,13 @@ FBuiltMember::FBuiltMember(FMemberId Name, FTypedRange Range)
 : FBuiltMember(Name, MoveTemp(Range.Schema), { .Range = Range.Values })
 {}
 
-FBuiltMember::FBuiltMember(FMemberId Name, FStructSchemaId Schema, FBuiltStruct* Value)
-: FBuiltMember(Name, MakeMemberSchema(DefaultStructType, Schema), { .Struct = Value })
+FBuiltMember::FBuiltMember(FMemberId Name, FStructId Schema, FBuiltStruct* Value)
+: FBuiltMember(Name, MakeMemberSchema(DefaultStructType, FInnerId(Schema)), { .Struct = Value })
 {}
 
-FBuiltMember FBuiltMember::MakeSuper(FStructSchemaId Schema, FBuiltStruct* Value)
+FBuiltMember FBuiltMember::MakeSuper(FStructId Schema, FBuiltStruct* Value)
 {
-	return FBuiltMember(NoId, MakeMemberSchema(SuperStructType, Schema), { .Struct = Value });
+	return FBuiltMember(NoId, MakeMemberSchema(SuperStructType, FInnerId(Schema)), { .Struct = Value });
 }
 
 //////////////////////////////////////////////////////////////////////////
diff --git a/Engine/Plugins/Experimental/PlainProps/Source/Private/PlainPropsBuildSchema.cpp b/Engine/Plugins/Experimental/PlainProps/Source/Private/PlainPropsBuildSchema.cpp
--- a/Engine/Plugins/Experimental/PlainProps/Source/Private/PlainPropsBuildSchema.cpp
+++ b/Engine/Plugins/Experimental/PlainProps/Source/Private/PlainPropsBuildSchema.cpp
@@ -13,18 +13,18 @@
 namespace PlainProps
 {
 
-static FString PrintMemberSchema(const FIds& Ids, FMemberType Type, FOptionalSchemaId InnerSchema, TConstArrayView<FMemberType> InnerRangeTypes)
+static FString PrintMemberSchema(const FIds& Ids, FMemberSchema Schema)
 {
 	TUtf8StringBuilder<256> Utf8SchemaStr;
-	PrintMemberSchema(Utf8SchemaStr, Ids, Type, InnerSchema, InnerRangeTypes);
+	PrintMemberSchema(Utf8SchemaStr, Ids, Schema);
 	FString SchemaStr(FStringView(StringCast<TCHAR>(Utf8SchemaStr.GetData(), Utf8SchemaStr.Len())));
-
-	FMemberType InnermostType = InnerRangeTypes.Num() > 0 ? InnerRangeTypes.Last() : Type;
+	
+	FMemberType InnermostType = Schema.GetInnermostType();
 	if (InnermostType.IsStruct())
 	{
 		return FString::Printf(TEXT("%sstruct [%d]%s%s => %s"),
-			Type.IsRange() ? TEXT("Range(s) of ") : TEXT(""),
-			InnerSchema.Get().Idx,
+			Schema.Type.IsRange() ? TEXT("Range(s) of ") : TEXT(""),
+			Schema.InnerSchema.Get().Idx,
 			InnermostType.AsStruct().IsSuper ? TEXT(" (super)") : TEXT(""),
 			InnermostType.AsStruct().IsDynamic ? TEXT(" (dynamic)") : TEXT(""),
 			*SchemaStr);
@@ -32,29 +32,24 @@ static FString PrintMemberSchema(const FIds& Ids, FMemberType Type, FOptionalSch
 	else if (IsEnum(InnermostType))
 	{
 		return FString::Printf(TEXT("%s%s [%d] => %s"),
-			Type.IsRange() ? TEXT("Range(s) of ") : TEXT(""),
+			Schema.Type.IsRange() ? TEXT("Range(s) of ") : TEXT(""),
 			StringCast<TCHAR>(ToString(InnermostType.AsLeaf()).GetData()).Get(),
-			InnerSchema.Get().Idx,
+			Schema.InnerSchema.Get().Idx,
 			*SchemaStr);
 	}
 	return SchemaStr;
 }
 
-static FString PrintMemberSchema(const FIds& Ids, FMemberSchema Schema)
-{
-	return PrintMemberSchema(Ids, Schema.Type, Schema.InnerSchema, Schema.GetInnerRangeTypes());
-}
-
 //////////////////////////////////////////////////////////////////////////
 
 struct FStructSchemaBuilder
 {
-	FStructSchemaBuilder(FStructSchemaId InId, const FStructDeclaration& InDecl, FSchemasBuilder& InSchemas);
+	FStructSchemaBuilder(FStructId InId, const FStructDeclaration& InDecl, FSchemasBuilder& InSchemas);
 
 	const FStructDeclaration&					Declaration;
 	FSchemasBuilder&							AllSchemas;
 	uint16										MinMembers = 0xFFFF;
-	FStructSchemaId								Id;
+	FStructId									Id;
 	FOptionalMemberId*							MemberOrder = nullptr;
 	FMemberSchema*								NotedSchemas = nullptr;
 	TBitArray<>									NotedMembers;
@@ -68,11 +63,11 @@ struct FEnumSchemaBuilder
 {
 	const FEnumDeclaration&						Declaration;
 	FDebugIds									Debug;
-	FEnumSchemaId								Id;
+	FEnumId										Id;
 	TOptional<ELeafWidth>						NotedWidth;
 	TSet<uint64>								NotedConstants;
 
-	void										NoteValue(ELeafWidth Width, uint64 Value, FStructSchemaId Struct, FMemberId Member);
+	void										NoteValue(ELeafWidth Width, uint64 Value, FStructId Struct, FMemberId Member);
 	void										NoteEmpty(ELeafWidth Width);
 	FBuiltEnumSchema							Build() const;
 };
@@ -110,26 +105,26 @@ FORCEINLINE BuilderType& GetOrEmplace(TSet<IdType>& Indexer, IdType Id, TPagedAr
 	return Builders.Last();
 }
 	
-FORCEINLINE FEnumSchemaBuilder&	FSchemasBuilder::NoteEnum(FEnumSchemaId Id)
+FORCEINLINE FEnumSchemaBuilder&	FSchemasBuilder::NoteEnum(FEnumId Id)
 {
 	checkf(!bBuilt, TEXT("Noted new members after building"));
 	checkf(Id.Idx < uint32(DeclaredEnums.Num()) && DeclaredEnums[Id.Idx], TEXT("Undeclared enum '%s' noted"), *Debug.Print(Id));
 	return GetOrEmplace(EnumIndices, Id, Enums, *DeclaredEnums[Id.Idx], Debug, Id);
 }
 
-FORCEINLINE FStructSchemaBuilder& FSchemasBuilder::NoteStruct(FStructSchemaId BindId)
+FORCEINLINE FStructSchemaBuilder& FSchemasBuilder::NoteStruct(FStructId BindId)
 {
 	checkf(!bBuilt, TEXT("Noted new members after building"));
 
-	FStructSchemaId DeclId	= BindId.Idx < uint32(DeclaredStructs.Num()) && DeclaredStructs[BindId.Idx] 
+	FStructId DeclId	= BindId.Idx < uint32(DeclaredStructs.Num()) && DeclaredStructs[BindId.Idx] 
 							? BindId : BindIds.GetDeclId(BindId);
 
 	checkf(DeclId.Idx < uint32(DeclaredStructs.Num()) && DeclaredStructs[DeclId.Idx], TEXT("Undeclared struct '%s' noted"), *Debug.Print(DeclId));
-	FStructSchemaId NoteId = (Format == ESchemaFormat::StableNames) ? DeclId : BindId;
+	FStructId NoteId = (Format == ESchemaFormat::StableNames) ? DeclId : BindId;
 	return GetOrEmplace(StructIndices, NoteId, Structs, NoteId, *DeclaredStructs[DeclId.Idx], *this);
 }
 
-void FSchemasBuilder::NoteStructAndMembers(FStructSchemaId BindId, const FBuiltStruct& Struct)
+void FSchemasBuilder::NoteStructAndMembers(FStructId BindId, const FBuiltStruct& Struct)
 {
 	NoteStruct(BindId).NoteMembersRecursively(Struct);
 }
@@ -161,17 +156,17 @@ void FSchemasBuilder::NoteInheritanceChains()
 {
 	for (int Idx = 0, Num = Structs.Num(); Idx < Num; ++Idx)
 	{
-		for (FOptionalStructSchemaId Super = Structs[Idx].Declaration.Super; Super; Super = DeclaredStructs[Super.Get().Idx]->Super)
+		for (FOptionalStructId Super = Structs[Idx].Declaration.Super; Super; Super = DeclaredStructs[Super.Get().Idx]->Super)
 		{
-			FStructSchemaId Id = Super.Get();
-			GetOrEmplace(StructIndices, Id, Structs, Super.Get(), *DeclaredStructs[Id.Idx], *this);
+			FStructId Id = Super.Get();
+			GetOrEmplace(StructIndices, Id, Structs, Id, *DeclaredStructs[Id.Idx], *this);
 		}
 	}
 }
 
 //////////////////////////////////////////////////////////////////////////
 
-FStructSchemaBuilder::FStructSchemaBuilder(FStructSchemaId InId, const FStructDeclaration& Decl, FSchemasBuilder& Schemas)
+FStructSchemaBuilder::FStructSchemaBuilder(FStructId InId, const FStructDeclaration& Decl, FSchemasBuilder& Schemas)
 : Declaration(Decl)
 , AllSchemas(Schemas)
 , Id(InId)
@@ -214,9 +209,9 @@ static void SetIsDynamic(FMemberType& InOut)
 	InOut = FMemberType(Type);
 }
 
-static void* NoteStructOrEnum(FSchemasBuilder& AllSchemas, bool bStruct, FSchemaId Id)
+static void* NoteStructOrEnum(FSchemasBuilder& AllSchemas, bool bStruct, FInnerId Id)
 {
-	return bStruct ? static_cast<void*>(&AllSchemas.NoteStruct(static_cast<FStructSchemaId>(Id))) : &AllSchemas.NoteEnum(static_cast<FEnumSchemaId>(Id));
+	return bStruct ? static_cast<void*>(&AllSchemas.NoteStruct(Id.AsStruct())) : &AllSchemas.NoteEnum(Id.AsEnum());
 }
 
 void FStructSchemaBuilder::NoteMembersRecursively(const FBuiltStruct& Struct)
@@ -274,16 +269,16 @@ void FStructSchemaBuilder::NoteMembersRecursively(const FBuiltStruct& Struct)
 		if (Schema.InnerSchema)
 		{
 			checkSlow(IsStructOrEnum(Schema.GetInnermostType()));
-			FSchemaId InnerSchema = Schema.InnerSchema.Get();
+			FInnerId InnerSchema = Schema.InnerSchema.Get();
 
 			switch (Schema.Type.GetKind())
 			{
 			case EMemberKind::Leaf:
-				AllSchemas.NoteEnum(static_cast<FEnumSchemaId>(InnerSchema)).NoteValue(Schema.Type.AsLeaf().Width, Member.Value.Leaf, Id, Member.Name.Get());
+				AllSchemas.NoteEnum(InnerSchema.AsEnum()).NoteValue(Schema.Type.AsLeaf().Width, Member.Value.Leaf, Id, Member.Name.Get());
 				break;
 
 			case EMemberKind::Struct:
-				AllSchemas.NoteStruct(static_cast<FStructSchemaId>(InnerSchema)).NoteMembersRecursively(*Member.Value.Struct);
+				AllSchemas.NoteStruct(InnerSchema.AsStruct()).NoteMembersRecursively(*Member.Value.Struct);
 				break;
 		
 			case EMemberKind::Range:
@@ -296,7 +291,7 @@ void FStructSchemaBuilder::NoteMembersRecursively(const FBuiltStruct& Struct)
 }
 
 template<typename IntType>
-void NoteEnumValues(FEnumSchemaBuilder& Schema, const IntType* Values, uint64 Num, FStructSchemaId Struct, FMemberId Member)
+void NoteEnumValues(FEnumSchemaBuilder& Schema, const IntType* Values, uint64 Num, FStructId Struct, FMemberId Member)
 {
 	for (IntType Value : TConstArrayView64<IntType>(Values, Num))
 	{
@@ -304,7 +299,7 @@ void NoteEnumValues(FEnumSchemaBuilder& Schema, const IntType* Values, uint64 Nu
 	}
 }
 
-static void NoteEnumRange(FEnumSchemaBuilder& Out, FLeafType Leaf, const FBuiltRange& Range, FStructSchemaId Struct, FMemberId Member)
+static void NoteEnumRange(FEnumSchemaBuilder& Out, FLeafType Leaf, const FBuiltRange& Range, FStructId Struct, FMemberId Member)
 {
 	check(Leaf.Type == ELeafType::Enum);
 	switch (Leaf.Width)
@@ -411,7 +406,7 @@ FBuiltEnumSchema FEnumSchemaBuilder::Build() const
 	return Out;
 }
 
-void FEnumSchemaBuilder::NoteValue(ELeafWidth Width, uint64 Value, FStructSchemaId Struct, FMemberId Member)
+void FEnumSchemaBuilder::NoteValue(ELeafWidth Width, uint64 Value, FStructId Struct, FMemberId Member)
 {
 	check(NotedWidth == NullOpt || NotedWidth == Width);
 	NotedWidth = Width;
@@ -459,12 +454,12 @@ void FEnumSchemaBuilder::NoteEmpty(ELeafWidth Width)
 
 //////////////////////////////////////////////////////////////////////////
 
-TArray<FStructSchemaId>	ExtractRuntimeIds(const FBuiltSchemas& In)
+TArray<FStructId>	ExtractRuntimeIds(const FBuiltSchemas& In)
 {
-	TArray<FStructSchemaId>	Out;
+	TArray<FStructId>	Out;
 	Out.SetNumUninitialized(In.Structs.Num());
 	const FBuiltStructSchema* InIt = In.Structs.GetData();
-	for (FStructSchemaId& Id : Out)
+	for (FStructId& Id : Out)
 	{
 		Id = (*InIt++).Id;
 	}
diff --git a/Engine/Plugins/Experimental/PlainProps/Source/Private/PlainPropsDeclare.cpp b/Engine/Plugins/Experimental/PlainProps/Source/Private/PlainPropsDeclare.cpp
--- a/Engine/Plugins/Experimental/PlainProps/Source/Private/PlainPropsDeclare.cpp
+++ b/Engine/Plugins/Experimental/PlainProps/Source/Private/PlainPropsDeclare.cpp
@@ -97,7 +97,7 @@ void CopyItems(T* It, TConstArrayView<T> Items)
 	}
 }
 
-static const FStructDeclaration& Declare(TArray<TUniquePtr<FStructDeclaration>>& Out, FStructSchemaId DeclId, FType Type, TConstArrayView<FMemberId> MemberOrder, EMemberPresence Occupancy, FOptionalStructSchemaId Super, uint16 NumMembers)
+static const FStructDeclaration& Declare(TArray<TUniquePtr<FStructDeclaration>>& Out, FStructId DeclId, FType Type, TConstArrayView<FMemberId> MemberOrder, EMemberPresence Occupancy, FOptionalStructId Super, uint16 NumMembers)
 {
 	if (static_cast<int32>(DeclId.Idx) >= Out.Num())
 	{
@@ -127,18 +127,18 @@ static const FStructDeclaration& Declare(TArray<TUniquePtr<FStructDeclaration>>&
 	return *Ptr;
 }
 
-const FStructDeclaration& FDeclarations::DeclareStruct(FStructSchemaId DeclId, FType Type, TConstArrayView<FMemberId> MemberOrder, EMemberPresence Occupancy, FOptionalStructSchemaId Super)
+const FStructDeclaration& FDeclarations::DeclareStruct(FStructId DeclId, FType Type, TConstArrayView<FMemberId> MemberOrder, EMemberPresence Occupancy, FOptionalStructId Super)
 {
 	checkf(!(Super && Occupancy == EMemberPresence::RequireAll), TEXT("'%s' is a dense substruct, this isn't supported, see BuildSuperStruct()"), *Debug.Print(DeclId));
 	return Declare(/* out */ DeclaredStructs, DeclId, Type, MemberOrder, Occupancy, Super, IntCastChecked<uint16>(MemberOrder.Num()));
 }
 
-const FStructDeclaration& FDeclarations::DeclareNumeralStruct(FStructSchemaId DeclId, FType Type, TConstArrayView<FMemberId> Numerals, EMemberPresence Occupancy)
+const FStructDeclaration& FDeclarations::DeclareNumeralStruct(FStructId DeclId, FType Type, TConstArrayView<FMemberId> Numerals, EMemberPresence Occupancy)
 {
 	return Declare(/* out */ DeclaredStructs, DeclId, Type, Numerals, Occupancy, NoId, IntCastChecked<uint16>(Numerals.Num()));
 }
 
-const FEnumDeclaration& FDeclarations::DeclareEnum(FEnumSchemaId Id, FType Type, EEnumMode Mode, TConstArrayView<FEnumerator> Enumerators, EEnumAliases Policy)
+const FEnumDeclaration& FDeclarations::DeclareEnum(FEnumId Id, FType Type, EEnumMode Mode, TConstArrayView<FEnumerator> Enumerators, EEnumAliases Policy)
 {
 	if (static_cast<int32>(Id.Idx) >= DeclaredEnums.Num())
 	{
@@ -164,7 +164,7 @@ const FEnumDeclaration& FDeclarations::DeclareEnum(FEnumSchemaId Id, FType Type,
 	return *Ptr;
 }
 
-void FDeclarations::DropStructRef(FStructSchemaId DeclId)
+void FDeclarations::DropStructRef(FStructId DeclId)
 {
 	Check(DeclId);
 	
@@ -177,18 +177,18 @@ void FDeclarations::DropStructRef(FStructSchemaId DeclId)
 }
 
 const FStructDeclaration*
-FDeclarations::Find(FStructSchemaId Id) const
+FDeclarations::Find(FStructId Id) const
 {
 	return Id.Idx < uint32(DeclaredStructs.Num()) ? DeclaredStructs[Id.Idx].Get() : nullptr;
 }
 
 #if DO_CHECK
-void FDeclarations::Check(FEnumSchemaId Id) const
+void FDeclarations::Check(FEnumId Id) const
 {
 	checkf(Id.Idx < (uint32)DeclaredEnums.Num() && DeclaredEnums[Id.Idx], TEXT("'%s' is undeclared"), *Debug.Print(Id));
 }
 
-void FDeclarations::Check(FStructSchemaId Id) const
+void FDeclarations::Check(FStructId Id) const
 {
 	checkf(Id.Idx < (uint32)DeclaredStructs.Num() && DeclaredStructs[Id.Idx], TEXT("'%s' is undeclared"), *Debug.Print(Id));
 }
diff --git a/Engine/Plugins/Experimental/PlainProps/Source/Private/PlainPropsDiff.cpp b/Engine/Plugins/Experimental/PlainProps/Source/Private/PlainPropsDiff.cpp
--- a/Engine/Plugins/Experimental/PlainProps/Source/Private/PlainPropsDiff.cpp
+++ b/Engine/Plugins/Experimental/PlainProps/Source/Private/PlainPropsDiff.cpp
@@ -9,7 +9,7 @@ namespace PlainProps
 {
 
 static bool DiffItem(const void* A, const void* B, FDiffContext& Ctx, FRangeMemberBinding Range);
-static bool DiffItem(const void* A, const void* B, FDiffContext& Ctx, FStructSchemaId Struct);
+static bool DiffItem(const void* A, const void* B, FDiffContext& Ctx, FStructId Struct);
 
 static bool DiffMember(FLeafMemberBinding Member, const void* OwnerA, const void* OwnerB, FDiffContext& Ctx, FOptionalMemberId Name)
 {
@@ -41,7 +41,7 @@ static bool DiffMember(FRangeMemberBinding Member, const void* OwnerA, const voi
 		const IItemRangeBinding& ItemBinding = Binding.AsItemBinding();
 		if (InnerKind == EMemberKind::Struct)
 		{
-			bDiff = DiffItemRange(A, B, ItemBinding, Ctx, static_cast<FStructSchemaId>(Member.InnermostSchema.Get()));
+			bDiff = DiffItemRange(A, B, ItemBinding, Ctx, Member.InnermostSchema.Get().AsStruct());
 		}
 		else if (InnerKind == EMemberKind::Leaf)
 		{
@@ -59,9 +59,9 @@ static bool DiffMember(FRangeMemberBinding Member, const void* OwnerA, const voi
 	return bDiff;
 }
 
-static bool DiffMembers(FStructSchemaId BindId, const void* A, const void* B, FDiffContext& Ctx);
+static bool DiffMembers(FStructId BindId, const void* A, const void* B, FDiffContext& Ctx);
 
-inline bool DiffStruct(FStructSchemaId BindId, const void* A, const void* B, FDiffContext& Ctx)
+inline bool DiffStruct(FStructId BindId, const void* A, const void* B, FDiffContext& Ctx)
 {
 	if (const ICustomBinding* Custom = Ctx.Customs.FindStruct(BindId))
 	{
@@ -94,7 +94,7 @@ static bool DiffMember(/* in-out */ FMemberVisitor& It, const void* A, const voi
 }
 
 
-static bool DiffMembers(FStructSchemaId BindId, const void* A, const void* B, FDiffContext& Ctx)
+static bool DiffMembers(FStructId BindId, const void* A, const void* B, FDiffContext& Ctx)
 {
 	const FSchemaBinding& Schema = Ctx.Schemas.GetStruct(BindId);
 	const FStructDeclaration& Declaration = Ctx.Types.Get(Schema.DeclId);
@@ -125,14 +125,14 @@ static bool DiffItem(const void* A, const void* B, FDiffContext& Ctx, FRangeMemb
 	return DiffMember(Range, A, B, Ctx, NoId);
 }
 
-static bool DiffItem(const void* A, const void* B, FDiffContext& Ctx, FStructSchemaId Id)
+static bool DiffItem(const void* A, const void* B, FDiffContext& Ctx, FStructId Id)
 {
 	return DiffStruct(Id, A, B, Ctx);
 }
 
 ////////////////////////////////////////////////////////////////////////////////////////////////
 
-bool DiffStructs(FStructSchemaId BindId, const void* A, const void* B, FDiffContext& Ctx)
+bool DiffStructs(FStructId BindId, const void* A, const void* B, FDiffContext& Ctx)
 {
 	if (const ICustomBinding* Custom = Ctx.Customs.FindStruct(BindId))
 	{
diff --git a/Engine/Plugins/Experimental/PlainProps/Source/Private/PlainPropsIndex.cpp b/Engine/Plugins/Experimental/PlainProps/Source/Private/PlainPropsIndex.cpp
--- a/Engine/Plugins/Experimental/PlainProps/Source/Private/PlainPropsIndex.cpp
+++ b/Engine/Plugins/Experimental/PlainProps/Source/Private/PlainPropsIndex.cpp
@@ -154,23 +154,23 @@ FType FIdIndexerBase::MakeAnonymousParametricType(TConstArrayView<FType> Params)
 	return { NoId, FTypenameId(MakeParametricTypeId(NoId, Params)) };
 }
 
-FEnumSchemaId FIdIndexerBase::IndexEnum(FType Type)
+FEnumId FIdIndexerBase::IndexEnum(FType Type)
 {
 	return { IntCastChecked<uint32>(Enums.Add(Type).AsInteger()) };
 }
 
-FStructSchemaId	FIdIndexerBase::IndexStruct(FType Type)
+FStructId FIdIndexerBase::IndexStruct(FType Type)
 {
 	return { IntCastChecked<uint32>(Structs.Add(Type).AsInteger()) };
 }
 
-FType FIdIndexerBase::Resolve(FEnumSchemaId Id) const
+FType FIdIndexerBase::Resolve(FEnumId Id) const
 {
 	check(Id.Idx < uint32(Enums.Num()));
 	return Enums.Get(FSetElementId::FromInteger(Id.Idx));
 }
 
-FType FIdIndexerBase::Resolve(FStructSchemaId Id) const
+FType FIdIndexerBase::Resolve(FStructId Id) const
 {
 	check(Id.Idx < uint32(Structs.Num()));
 	return Structs.Get(FSetElementId::FromInteger(Id.Idx));
diff --git a/Engine/Plugins/Experimental/PlainProps/Source/Private/PlainPropsInternalBind.h b/Engine/Plugins/Experimental/PlainProps/Source/Private/PlainPropsInternalBind.h
--- a/Engine/Plugins/Experimental/PlainProps/Source/Private/PlainPropsInternalBind.h
+++ b/Engine/Plugins/Experimental/PlainProps/Source/Private/PlainPropsInternalBind.h
@@ -23,7 +23,7 @@ public:
 	FLeafMemberBinding			GrabLeaf();				// @pre PeekKind() == EMemberKind::Leaf
 	FRangeMemberBinding			GrabRange();			// @pre PeekKind() == EMemberKind::Range
 	FStructMemberBinding		GrabStruct();			// @pre PeekKind() == EMemberKind::Struct
-	FStructSchemaId				GrabSuper();			// @pre First grab and has declared super
+	FStructId					GrabSuper();			// @pre First grab and has declared super
 	void						SkipMember();
 
 protected: // for unit tests
@@ -37,31 +37,29 @@ protected: // for unit tests
 
 	uint64						GrabMemberOffset();
 	FMemberBindTypeRange		GrabInnerTypes();
-	FSchemaId					GrabInnerSchema();
-	FStructSchemaId				GrabStructSchema(FStructType Type);
-	FOptionalSchemaId			GrabRangeSchema(FMemberType InnermostType);
-	FEnumSchemaId				GrabEnumSchema()			{ return static_cast<FEnumSchemaId&&>(GrabInnerSchema()); }
+	FInnerId					GrabInnerSchema();
+	FStructId					GrabStructSchema(FStructType Type);
+	FOptionalInnerId			GrabRangeSchema(FMemberType InnermostType);
+	FEnumId						GrabEnumSchema()		{ return GrabInnerSchema().AsEnum(); }
 };
 
 ////////////////////////////////////////////////////////////////////////////////////////////////
 
-struct FMemberBinder
+struct FMemberBinderBase
 {
-	FMemberBinder(FSchemaBinding& InSchema)
+	FMemberBinderBase(FSchemaBinding& InSchema)
 	: Schema(InSchema)
 	, MemberIt(Schema.Members)
 	, RangeTypeIt(const_cast<FMemberBindType*>(Schema.GetInnerRangeTypes()))
 	, OffsetIt(const_cast<uint32*>(Schema.GetOffsets()))
-	, InnerSchemaIt(const_cast<FSchemaId*>(Schema.GetInnerSchemas()))
 	, RangeBindingIt(const_cast<FRangeBinding*>(Schema.GetRangeBindings()))
 	{}
 
-	~FMemberBinder()
+	~FMemberBinderBase()
 	{
 		check(MemberIt == Schema.GetInnerRangeTypes());
 		check(Align(RangeTypeIt, alignof(uint32)) == (const void*)Schema.GetOffsets());
 		check(OffsetIt == (const void*)Schema.GetInnerSchemas());
-		check(Align(InnerSchemaIt, alignof(FRangeBinding)) == (const void*)Schema.GetRangeBindings() || Schema.NumInnerRanges == 0);
 		check(Schema.NumInnerRanges == RangeBindingIt - Schema.GetRangeBindings());
 	}
 
@@ -85,16 +83,10 @@ struct FMemberBinder
 		RangeBindingIt += Ranges.Num();
 	}
 	
-	void AddInnerSchema(FSchemaId InnermostSchema)
-	{
-		*InnerSchemaIt++ = InnermostSchema;
-	}
-
 	FSchemaBinding& Schema;
 	FMemberBindType* MemberIt;
 	FMemberBindType* RangeTypeIt;
 	uint32* OffsetIt;
-	FSchemaId* InnerSchemaIt;
 	FRangeBinding* RangeBindingIt;
 };
 
diff --git a/Engine/Plugins/Experimental/PlainProps/Source/Private/PlainPropsInternalParse.h b/Engine/Plugins/Experimental/PlainProps/Source/Private/PlainPropsInternalParse.h
--- a/Engine/Plugins/Experimental/PlainProps/Source/Private/PlainPropsInternalParse.h
+++ b/Engine/Plugins/Experimental/PlainProps/Source/Private/PlainPropsInternalParse.h
@@ -64,22 +64,22 @@ struct FToken
 struct FParsedMemberSchema
 {
 	FMemberType						Type;
-	FOptionalSchemaId				InnerSchema;
+	FOptionalInnerId				InnerSchema;
 	TArray<FMemberType>				InnerRangeTypes;
 };
 
 struct FParsedStructSchema
 {
-	FStructSchemaId					Id;
-	FOptionalStructSchemaId			DeclaredSuper;
-	FOptionalStructSchemaId			UsedSuper;
+	FStructId						Id;
+	FOptionalStructId				DeclaredSuper;
+	FOptionalStructId				UsedSuper;
 	TConstArrayView<FMemberId>		MemberNames;	// FStructDeclaration.MemberOrder
 	TArray<FParsedMemberSchema>		MemberSchemas;	// Same size as declared member names
 };
 
 struct FParsedEnumSchema
 {
-	FEnumSchemaId					Id;
+	FEnumId							Id;
 	ELeafWidth						Width;
 	TConstArrayView<FEnumerator>	Enumerators;	// FEnumDeclaration.Enumerators
 };
@@ -97,22 +97,22 @@ struct FParsedSchemas
 class FTextIndexer : public TIdIndexer<FUtf8StringView>
 {
 public:
-	FOptionalStructSchemaId GetStructId(FType Type) const
+	FOptionalStructId GetStructId(FType Type) const
 	{
 		FSetElementId Id = Structs.FindId(Type);
 		if (Id.IsValidId())
 		{
-			return FStructSchemaId{ IntCastChecked<uint32>(Id.AsInteger()) };
+			return FStructId{ IntCastChecked<uint32>(Id.AsInteger()) };
 		}
 		return NoId;
 	}
 
-	FOptionalEnumSchemaId GetEnumId(FType Type) const
+	FOptionalEnumId GetEnumId(FType Type) const
 	{
 		FSetElementId Id = Enums.FindId(Type);
 		if (Id.IsValidId())
 		{
-			return FEnumSchemaId{ IntCastChecked<uint32>(Id.AsInteger()) };
+			return FEnumId{ IntCastChecked<uint32>(Id.AsInteger()) };
 		}
 		return NoId;
 	}
@@ -155,8 +155,8 @@ private:
 	FScratchAllocator&								Scratch;
 	FTextIndexer									Names;
 	FDeclarations									Types;
-	TArray<TTuple<FStructSchemaId, FBuiltStruct*>>	Objects;
+	TArray<TTuple<FStructId, FBuiltStruct*>>		Objects;
 	FParsedSchemas									Schemas;
 };
 
-} // namespace PlainProps
+} // namespace PlainProps
\ No newline at end of file
diff --git a/Engine/Plugins/Experimental/PlainProps/Source/Private/PlainPropsInternalPrint.h b/Engine/Plugins/Experimental/PlainProps/Source/Private/PlainPropsInternalPrint.h
--- a/Engine/Plugins/Experimental/PlainProps/Source/Private/PlainPropsInternalPrint.h
+++ b/Engine/Plugins/Experimental/PlainProps/Source/Private/PlainPropsInternalPrint.h
@@ -30,14 +30,13 @@ using FYamlBuilderPtr = TUniquePtr<FYamlBuilder, FYamlBuilderDeleter>;
 class FBatchPrinter
 {
 public:
-	FBatchPrinter(FYamlBuilder& InTextBuilder, const FIds& InIds);
+	FBatchPrinter(FYamlBuilder& InTextBuilder, const FBatchIds& InIds);
 	~FBatchPrinter();
 
-	void PrintSchemas(FReadBatchId Batch);
+	void PrintSchemas();
 	void PrintObjects(TConstArrayView<FStructView> Structs);
 
 private:
-	void PrintSchemas(const FSchemaBatch& Batch, FReadBatchId BatchId);
 	void PrintStructSchema(const FStructSchema& Struct, FReadBatchId BatchId);
 	void PrintEnumSchema(const FEnumSchema& Enum);
 
@@ -45,7 +44,7 @@ private:
 	void PrintEnumConstants(TConstArrayView<FNameId> EnumNames, TConstArrayView<IntType> Constants, bool bFlagMode);
 
 	FYamlBuilder& TextBuilder;
-	const FIds& Ids;
+	const FBatchIds& Ids;
 };
 
 ///////////////////////////////////////////////////////////////////////////////
diff --git a/Engine/Plugins/Experimental/PlainProps/Source/Private/PlainPropsInternalRead.h b/Engine/Plugins/Experimental/PlainProps/Source/Private/PlainPropsInternalRead.h
--- a/Engine/Plugins/Experimental/PlainProps/Source/Private/PlainPropsInternalRead.h
+++ b/Engine/Plugins/Experimental/PlainProps/Source/Private/PlainPropsInternalRead.h
@@ -50,6 +50,13 @@ inline FMemoryView GrabRangeValues(uint64 Num, FMemberType InnerType, /* in-out
 	return ByteIt.GrabSlice(NumBytes);
 }
 
+inline FMemberType GetInnermostType(FRangeSchema Schema)
+{
+	FMemberType Inner = Schema.ItemType;
+	for (const FMemberType* It = Schema.NestedItemTypes; Inner.IsRange(); Inner = *It++) {}
+	return Inner;
+}
+
 //////////////////////////////////////////////////////////////////////////
 
 template<class T>
diff --git a/Engine/Plugins/Experimental/PlainProps/Source/Private/PlainPropsLoad.cpp b/Engine/Plugins/Experimental/PlainProps/Source/Private/PlainPropsLoad.cpp
--- a/Engine/Plugins/Experimental/PlainProps/Source/Private/PlainPropsLoad.cpp
+++ b/Engine/Plugins/Experimental/PlainProps/Source/Private/PlainPropsLoad.cpp
@@ -15,14 +15,27 @@
 namespace PlainProps
 {
 
-struct FLoadStructMemcpy
+struct FMemcpyLoadPlan
 {
 	uint32					Size;
 	uint32					Offset;
 };
 
+// Duplicated runtime FSchemaBinding with FStructSchemaId load ids instead of runtime FInnerId
+struct FSchemaLoadPlan
+{
+	const FSchemaBinding&	Clone;
+
+	TConstArrayView<FMemberBindType>	GetMembers() const			{ return MakeArrayView(Clone.Members, Clone.NumMembers); }
+	const uint32*						GetOffsets() const			{ return Clone.GetOffsets(); }
+	TConstArrayView<FStructSchemaId>	GetInnerSchemas() const		{ return MakeArrayView(reinterpret_cast<const FStructSchemaId*>(Clone.GetInnerSchemas()), Clone.NumInnerSchemas); }
+	TConstArrayView<FMemberBindType>	GetInnerRangeTypes() const	{ return MakeArrayView(Clone.GetInnerRangeTypes(), Clone.NumInnerRanges); }
+	const FRangeBinding*				GetRangeBindings() const	{ return Clone.GetRangeBindings(); }
+
+};
+
 // ICustomBinding with internal type-erased/lowered struct schemas
-struct FCustomLoadIdsPlan
+struct FCustomLoadPlan
 {
 	const ICustomBinding*	Binding;
 	uint32					NumLoadIds;
@@ -35,7 +48,7 @@ class FLoadStructPlan
 public:
 	FLoadStructPlan() = default;
 
-	explicit FLoadStructPlan(FLoadStructMemcpy Memcpy)
+	explicit FLoadStructPlan(FMemcpyLoadPlan Memcpy)
 	: Handle((uint64(Memcpy.Size) << 32) | (uint64(Memcpy.Offset) << 2) | MemcpyMask)
 	{
 		check(Memcpy.Offset == AsMemcpy().Offset && Memcpy.Size == AsMemcpy().Size);
@@ -47,7 +60,7 @@ public:
 		check(&Custom == &AsCustom());
 	}
 
-	explicit FLoadStructPlan(const FCustomLoadIdsPlan& Custom)
+	explicit FLoadStructPlan(const FCustomLoadPlan& Custom)
 	: Handle(uint64(&Custom) | CustomMask | LoadIdsBit)
 	{
 		check(Custom.Binding == &AsCustom());
@@ -59,7 +72,7 @@ public:
 	: Handle(uint64(&Schema) | SchemaBit | (bSparse ? SparseBit : 0) | (uint64(OffsetWidth) << SchemaOffsetShift))
 	{
 		static_assert(alignof(FSchemaBinding) >= 8);
-		check(&Schema == &AsSchema());
+		check(&Schema == &AsSchema().Clone);
 		check(bSparse == IsSparseSchema());
 		check(OffsetWidth == GetOffsetWidth());
 	}
@@ -68,11 +81,11 @@ public:
 	bool						IsSparseSchema() const	{ return (Handle & SparseSchemaMask) == SparseSchemaMask; }
 	bool						IsMemcpy() const		{ return (Handle & LoMask) == MemcpyMask; }
 	bool						IsCustom() const		{ return (Handle & LoMask) == CustomMask; }
-	FLoadStructMemcpy			AsMemcpy() const		{ check(IsMemcpy()); return { static_cast<uint32>(Handle >> 32), static_cast<uint32>(Handle) >> 2 }; }
-	const FSchemaBinding&		AsSchema() const		{ check(IsSchema()); return *AsPtr<FSchemaBinding>(); }
+	FMemcpyLoadPlan				AsMemcpy() const		{ check(IsMemcpy()); return { static_cast<uint32>(Handle >> 32), static_cast<uint32>(Handle) >> 2 }; }
+	FSchemaLoadPlan				AsSchema() const		{ check(IsSchema()); return { *AsPtr<FSchemaBinding>() }; }
 	const ICustomBinding&		AsCustom() const;
 
-	const FStructSchemaId*		GetInnerLoadIds()		{ return (Handle & TagMask) == LoadIdsMask ? AsPtr<FCustomLoadIdsPlan>()->LoadIds : nullptr; }
+	const FStructSchemaId*		GetInnerLoadIds()		{ return (Handle & TagMask) == LoadIdsMask ? AsPtr<FCustomLoadPlan>()->LoadIds : nullptr; }
 	ELeafWidth					GetOffsetWidth() const	{ check(IsSchema()); return static_cast<ELeafWidth>((Handle & SchemaOffsetMask) >> SchemaOffsetShift); }
 
 private:
@@ -102,10 +115,9 @@ private:
 const ICustomBinding& FLoadStructPlan::AsCustom() const
 { 
 	check(IsCustom());
-	return (Handle & LoadIdsBit) ? *AsPtr<FCustomLoadIdsPlan>()->Binding : *AsPtr<ICustomBinding>();
+	return (Handle & LoadIdsBit) ? *AsPtr<FCustomLoadPlan>()->Binding : *AsPtr<ICustomBinding>();
 }
 
-
 ////////////////////////////////////////////////////////////////////////////
 
 static uint16 CountEnums(const FStructSchema& Schema)
@@ -143,7 +155,7 @@ static uint16 CountEnums(const FStructSchema& Schema)
 	return Num;
 }
 
-static bool HasDifferentSupers(const FStructSchema& From, const FSchemaBinding& To, TConstArrayView<FStructSchemaId> ToStructIds)
+static bool HasDifferentSupers(const FStructSchema& From, const FSchemaBinding& To, TConstArrayView<FStructId> ToStructIds)
 {
 	if (From.Inheritance == ESuper::No)
 	{
@@ -151,8 +163,8 @@ static bool HasDifferentSupers(const FStructSchema& From, const FSchemaBinding&
 	}
 	else if (To.HasSuper())
 	{
-		FStructSchemaId FromSuper = ToStructIds[From.GetSuperSchema().Get().Idx];
-		FStructSchemaId ToSuper = static_cast<FStructSchemaId>(To.GetInnerSchemas()[0]);
+		FStructId FromSuper = ToStructIds[From.GetSuperSchema().Get().Idx];
+		FStructId ToSuper = To.GetInnerSchemas()[0].AsStruct();
 		return FromSuper == ToSuper;
 	}
 
@@ -166,22 +178,22 @@ struct FLoadIdMapping
 {
 	FStructSchemaId ReadId; // ~ Batch decl id, index into saved schemas and load plans
 	FStructSchemaId LoadId; // ~ Batch bind id, index into load plans
-	FStructSchemaId BindId; // Runtime bind id
+	FStructId		BindId; // Runtime bind id
 };
 
 // Helps load type-erased/lowered structs with ExplicitBindName by allocating new load-time struct ids
 class FLoadIdBinder
 {
 public:
-	FLoadIdBinder(TConstArrayView<FStructSchemaId> RuntimeDeclIds, FDebugIds Dbg)
+	FLoadIdBinder(TConstArrayView<FStructId> RuntimeDeclIds, FDebugIds Dbg)
 	: DeclIds(RuntimeDeclIds)
 	, NextLoadIdx(static_cast<uint32>(RuntimeDeclIds.Num()))
 	, Debug(Dbg)
 	{}
 
-	FStructSchemaId					BindLoadId(FStructSchemaId ReadId, FStructSchemaId BindId)
+	FStructSchemaId					BindLoadId(FStructSchemaId ReadId, FStructId BindId)
 	{
-		FStructSchemaId DeclId = DeclIds[ReadId.Idx];
+		FStructId DeclId = DeclIds[ReadId.Idx];
 		FStructSchemaId LoadId = BindId == DeclId ? ReadId : MapLoadId(ReadId, BindId);
 		return LoadId;
 	}
@@ -189,11 +201,11 @@ public:
 	int32							NumMappings() const { return Mappings.Num(); }
 
 private:
-	TConstArrayView<FStructSchemaId>	DeclIds;
+	TConstArrayView<FStructId>			DeclIds;
 	uint32								NextLoadIdx;
 	TArray<FLoadIdMapping>				Mappings;
 
-	FStructSchemaId MapLoadId(FStructSchemaId ReadId, FStructSchemaId BindId)
+	FStructSchemaId MapLoadId(FStructSchemaId ReadId, FStructId BindId)
 	{
 		for (FLoadIdMapping Mapping : Mappings)
 		{
@@ -217,7 +229,7 @@ public:
 struct FLoadIdBinderDummy
 {
 	FDebugIds							Debug;
-	static constexpr FStructSchemaId	BindLoadId(FStructSchemaId ReadId, FStructSchemaId BindId) { return ReadId; }
+	static constexpr FStructSchemaId	BindLoadId(FStructSchemaId ReadId, FStructId BindId) { return ReadId; }
 	static constexpr int32				NumMappings() { return 0; }
 	static inline FLoadIdMapping		GetMapping(int32 Idx) { unimplemented(); return {}; }
 };
@@ -253,7 +265,30 @@ void FLoadBatchDeleter::operator()(FLoadBatch* Batch) const
 
 using SubsetByteArray = TArray<uint8, TInlineAllocator<1024>>;
 
-static void CopyMemberBinding(FLeafMemberBinding Binding, /* in-out */ const FSchemaId*& InnerSchemaIt, FMemberBinder& Out)
+////////////////////////////////////////////////////////////////////////////
+
+struct FMemberLoadBinder : FMemberBinderBase
+{
+	FMemberLoadBinder(FSchemaBinding& In)
+	: FMemberBinderBase(In)
+	, InnerSchemaIt(const_cast<FStructSchemaId*>(FSchemaLoadPlan(In).GetInnerSchemas().GetData()))
+	{}
+	
+	~FMemberLoadBinder()
+	{
+		check(Align(InnerSchemaIt, alignof(FRangeBinding)) == (const void*)Schema.GetRangeBindings() || Schema.NumInnerRanges == 0);
+	}
+
+	void AddInnerSchema(FStructSchemaId InnermostSchema)
+	{
+		*InnerSchemaIt++ = InnermostSchema;
+	}
+
+	FStructSchemaId* InnerSchemaIt;
+};
+
+
+static void CopyMemberBinding(FLeafMemberBinding Binding, /* in-out */ const FSchemaId*& InnerSchemaIt, FMemberLoadBinder& Out)
 {
 	InnerSchemaIt += (Binding.Leaf.Type == ELeafBindType::Enum); // Skip enum schema
 	Out.AddMember(Binding.Leaf.Pack(), static_cast<uint32>(Binding.Offset));
@@ -262,7 +297,7 @@ static void CopyMemberBinding(FLeafMemberBinding Binding, /* in-out */ const FSc
 // TODO: Make all below free functions FLoadPlanner member functions and make LoadIds a member variable
 
 template<class LoadIdBinder>
-void CopyMemberBinding(FStructMemberBinding Binding, /* in-out */ const FSchemaId*& InnerSchemaIt, LoadIdBinder& LoadIds, FMemberBinder& Out)
+void CopyMemberBinding(FStructMemberBinding Binding, /* in-out */ const FSchemaId*& InnerSchemaIt, LoadIdBinder& LoadIds, FMemberLoadBinder& Out)
 {
 	check(!Binding.Type.IsSuper);
 	check(!Binding.Type.IsDynamic); // TODO: Build out dynamic struct support
@@ -273,13 +308,14 @@ void CopyMemberBinding(FStructMemberBinding Binding, /* in-out */ const FSchemaI
 }
 
 template<class LoadIdBinder>
-void CopyMemberBinding(FRangeMemberBinding Binding, /* in-out */ const FSchemaId*& InnerSchemaIt, LoadIdBinder& LoadIds, FMemberBinder& Out)
+void CopyMemberBinding(FRangeMemberBinding Binding, /* in-out */ const FSchemaId*& InnerSchemaIt, LoadIdBinder& LoadIds, FMemberLoadBinder& Out)
 {
 	FMemberBindType InnermostType = Binding.InnerTypes[Binding.NumRanges - 1];
 	Out.AddRange(MakeArrayView(Binding.RangeBindings, Binding.NumRanges), InnermostType, static_cast<uint32>(Binding.Offset));
 	if (InnermostType.IsStruct())
 	{
-		FStructSchemaId LoadId = LoadIds.BindLoadId(static_cast<FStructSchemaId>(*InnerSchemaIt), static_cast<FStructSchemaId>(Binding.InnermostSchema.Get()));
+		FStructSchemaId ReadId = static_cast<FStructSchemaId>(*InnerSchemaIt);
+		FStructSchemaId LoadId = LoadIds.BindLoadId(ReadId, Binding.InnermostSchema.Get().AsStruct());
 		Out.AddInnerSchema(LoadId);
 		++InnerSchemaIt;
 	}
@@ -290,7 +326,7 @@ void CopyMemberBinding(FRangeMemberBinding Binding, /* in-out */ const FSchemaId
 }
 
 template<class LoadIdBinder>
-void CopyMemberBinding(/* in-out */ FMemberVisitor& BindIt, /* in-out */ const FSchemaId*& InnerSchemaIt, LoadIdBinder& LoadIds, FMemberBinder& Out)
+void CopyMemberBinding(/* in-out */ FMemberVisitor& BindIt, /* in-out */ const FSchemaId*& InnerSchemaIt, LoadIdBinder& LoadIds, FMemberLoadBinder& Out)
 {
 	switch (BindIt.PeekKind())
 	{
@@ -301,6 +337,7 @@ void CopyMemberBinding(/* in-out */ FMemberVisitor& BindIt, /* in-out */ const F
 	}
 }
 
+
 template<class LoadIdBinder>
 static void CreateSubsetBindingWithoutEnumIds(const FStructSchema& From, const FSchemaBinding& To, TConstArrayView<FMemberId> ToNames, uint16 NumEnums, LoadIdBinder& LoadIds, SubsetByteArray& Out)
 {
@@ -318,13 +355,13 @@ static void CreateSubsetBindingWithoutEnumIds(const FStructSchema& From, const F
 	
 	// Copy subset of member bindings...
 	FMemberVisitor ToIt(To);
-	FMemberBinder Footer(*Schema);
+	FMemberLoadBinder Footer(*Schema);
 	const FSchemaId* InnerSchemaIt = From.GetInnerSchemas() + bSkipSuperSchema;
 
 	// ...first the unnamed super member...
 	if (From.Inheritance != ESuper::No)
 	{
-		FStructSchemaId BindId = ToIt.GrabSuper();
+		FStructId BindId = ToIt.GrabSuper();
 		if (UsesSuper(From.Inheritance))
 		{
 			FStructType FromType = From.Footer[0].AsStruct(); // To.Members[0].AsStruct().IsDynamic isn't set
@@ -356,6 +393,7 @@ static void CreateSubsetBindingWithoutEnumIds(const FStructSchema& From, const F
 	check(InnerSchemaIt == From.GetInnerSchemas() + From.NumInnerSchemas);
 }
 
+// @pre No enum members
 template<class LoadIdBinder>
 static void CloneBindingWithReplacedStructIds(const FSchemaId* FromIds, const FSchemaBinding& To, LoadIdBinder& LoadIds, SubsetByteArray& Out)
 {
@@ -367,18 +405,18 @@ static void CloneBindingWithReplacedStructIds(const FSchemaId* FromIds, const FS
 	FSchemaBinding* Schema = reinterpret_cast<FSchemaBinding*>(&Out[Out.Num() - Size]);
 	FMemory::Memcpy(Schema, &To, Size);
 
-	// Replace inner schema bind ids with batch load ids
-	const FStructSchemaId* SaveIdIt = static_cast<const FStructSchemaId*>(FromIds);
-	FStructSchemaId* OutIds = static_cast<FStructSchemaId*>(const_cast<FSchemaId*>(Schema->GetInnerSchemas()));
-	for (FStructSchemaId& OutId : MakeArrayView(OutIds, To.NumInnerSchemas))
+	// Replace inner bind ids with batch load ids
+	const FStructSchemaId* ReadIdIt = static_cast<const FStructSchemaId*>(FromIds);
+	FInnerId* InnerIt = const_cast<FInnerId*>(Schema->GetInnerSchemas());
+	for (FInnerId& Inner : MakeArrayView(InnerIt, To.NumInnerSchemas))
 	{
-		FStructSchemaId MemcopiedBindId = OutId;
-		OutId = LoadIds.BindLoadId(*SaveIdIt++, MemcopiedBindId);
+		FStructId MemcopiedBindId = Inner.AsStruct();
+		reinterpret_cast<FStructSchemaId&>(Inner) = LoadIds.BindLoadId(*ReadIdIt++, MemcopiedBindId);
 	}
 }
 
 template<class LoadIdBinder>
-[[nodiscard]] static FLoadStructPlan MakeSchemaLoadPlan(const FStructSchema& From, const FSchemaBinding& To, TConstArrayView<FMemberId> ToMemberIds, TConstArrayView<FStructSchemaId> ToStructIds, LoadIdBinder& LoadIds, SubsetByteArray& OutSubsetSchemas)
+[[nodiscard]] static FLoadStructPlan MakeSchemaLoadPlan(const FStructSchema& From, const FSchemaBinding& To, TConstArrayView<FMemberId> ToMemberIds, TConstArrayView<FStructId> ToStructIds, LoadIdBinder& LoadIds, SubsetByteArray& OutSubsetSchemas)
 {
 	uint16 NumEnums = CountEnums(From);
 	if (From.NumMembers < To.NumMembers || NumEnums || HasDifferentSupers(From, To, ToStructIds))
@@ -402,7 +440,7 @@ template<class LoadIdBinder>
 	return FLoadStructPlan(To, ELeafWidth::B32, !From.IsDense);
 }
 
-[[nodiscard]] static TOptional<FLoadStructMemcpy> TryMakeMemcpyPlan(const FStructSchema& From, const FSchemaBinding& To, TConstArrayView<FMemberId> ToNames)
+[[nodiscard]] static TOptional<FMemcpyLoadPlan> TryMakeMemcpyPlan(const FStructSchema& From, const FSchemaBinding& To, TConstArrayView<FMemberId> ToNames)
 {
 	// Can't memcopy sparse members nor range-bound members nor super structs
 	if (!From.IsDense || From.NumRangeTypes > 0 || UsesSuper(From.Inheritance))
@@ -437,13 +475,13 @@ template<class LoadIdBinder>
 		EndPos += SizeOf(Member.AsLeaf().Width);
 	}
 
-	return FLoadStructMemcpy(EndPos - StartPos, StartPos);
+	return FMemcpyLoadPlan(EndPos - StartPos, StartPos);
 }
 
 template<class LoadIdBinder>
-[[nodiscard]] static FLoadStructPlan MakeLoadPlan(const FStructSchema& From, const FSchemaBinding& To, TConstArrayView<FMemberId> ToMemberIds, TConstArrayView<FStructSchemaId> ToStructIds, LoadIdBinder& LoadIds, SubsetByteArray& OutSubsetSchemas)
+[[nodiscard]] static FLoadStructPlan MakeLoadPlan(const FStructSchema& From, const FSchemaBinding& To, TConstArrayView<FMemberId> ToMemberIds, TConstArrayView<FStructId> ToStructIds, LoadIdBinder& LoadIds, SubsetByteArray& OutSubsetSchemas)
 {
-	TOptional<FLoadStructMemcpy> Memcpy = TryMakeMemcpyPlan(From, To, ToMemberIds);
+	TOptional<FMemcpyLoadPlan> Memcpy = TryMakeMemcpyPlan(From, To, ToMemberIds);
 	return Memcpy ? FLoadStructPlan(Memcpy.GetValue()) : MakeSchemaLoadPlan(From, To, ToMemberIds, ToStructIds, LoadIds, OutSubsetSchemas);
 }
 
@@ -500,7 +538,7 @@ struct FLoadPlanner
 	const FDeclarations&							Declarations;
 	const FCustomBindings&							Customs;
 	const FSchemaBindings&							Schemas;
-	const TConstArrayView<FStructSchemaId>			RuntimeIds;
+	const TConstArrayView<FStructId>				RuntimeIds;
 
 	// Temporary data structures	
 	TArray<FLoadStructPlan, TInlineAllocator<256>>	Plans;
@@ -617,8 +655,8 @@ private:
 					else
 					{
 						check(Plans[Idx].IsCustom());
-						check(IsAligned(Size, alignof(FCustomLoadIdsPlan)));
-						Out->Plans[Idx] = FLoadStructPlan(*reinterpret_cast<const FCustomLoadIdsPlan*>(It));
+						check(IsAligned(Size, alignof(FCustomLoadPlan)));
+						Out->Plans[Idx] = FLoadStructPlan(*reinterpret_cast<const FCustomLoadPlan*>(It));
 					}
 					It += Size;
 				}
@@ -679,21 +717,21 @@ private:
 		{}
 	};
 
-	[[nodiscard]] FCustomLoadIdsPlan* CreateCustomLoadIdsPlan(const ICustomBinding* Custom, const FStructSchema& Schema, TConstArrayView<FInnerStruct> Inners, /* in-out */ FLoadIdBinder& LoadIds)
+	[[nodiscard]] FCustomLoadPlan* CreateCustomLoadIdsPlan(const ICustomBinding* Custom, const FStructSchema& Schema, TConstArrayView<FInnerStruct> Inners, /* in-out */ FLoadIdBinder& LoadIds)
 	{
 		const uint32 Num = static_cast<uint32>(Inners.Num());
 		check(Num);
 		check(Schema.NumInnerSchemas > 0);
 
 		// Calculate padding and size needed
-		uint32 Size = offsetof(FCustomLoadIdsPlan, LoadIds) + sizeof(FStructSchemaId) * Num;
-		uint32 Pad = Align(SubsetSchemaData.Num(), alignof(FCustomLoadIdsPlan)) - SubsetSchemaData.Num();
+		uint32 Size = offsetof(FCustomLoadPlan, LoadIds) + sizeof(FStructSchemaId) * Num;
+		uint32 Pad = Align(SubsetSchemaData.Num(), alignof(FCustomLoadPlan)) - SubsetSchemaData.Num();
 		SubsetSchemaData.AddUninitialized(Pad + Size);
 
 		// Construct and init header
-		FCustomLoadIdsPlan Header = {Custom, Num};
-		FCustomLoadIdsPlan* Out = new (&SubsetSchemaData[SubsetSchemaData.Num() - Size]) FCustomLoadIdsPlan {Header};
-		check(IsAligned(Out, alignof(FCustomLoadIdsPlan)));
+		FCustomLoadPlan Header = {Custom, Num};
+		FCustomLoadPlan* Out = new (&SubsetSchemaData[SubsetSchemaData.Num() - Size]) FCustomLoadPlan {Header};
+		check(IsAligned(Out, alignof(FCustomLoadPlan)));
 		FStructSchemaId* OutLoadIdIt = Out->LoadIds;
 	
 		// Populate FLoadIdMapping and init Out->LoadIds
@@ -793,12 +831,11 @@ private:
 	}
 };
 
-FLoadBatchPtr CreateLoadPlans(FReadBatchId ReadId, const FDeclarations& Declarations, const FCustomBindings& Customs, const FSchemaBindings& Schemas, TConstArrayView<FStructSchemaId> RuntimeIds, ESchemaFormat Format)
+FLoadBatchPtr CreateLoadPlans(FReadBatchId ReadId, const FDeclarations& Declarations, const FCustomBindings& Customs, const FSchemaBindings& Schemas, TConstArrayView<FStructId> RuntimeIds, ESchemaFormat Format)
 {
 	return FLoadPlanner{ReadId, Declarations, Customs, Schemas, RuntimeIds}.CreatePlans(Format);
 }
 
-
 ////////////////////////////////////////////////////////////////////////////
 
 inline void SetBit(uint8& Out, uint8 Idx, bool bValue)
@@ -991,11 +1028,11 @@ template<bool bSparse, typename OffsetType>
 class TMemberLoader : public FRangeLoader
 {
 public:
-	TMemberLoader(FByteReader Values, const FSchemaBinding& Schema, const FLoadBatch& InBatch)
-	: Types(Schema.Members, Schema.NumMembers)
+	TMemberLoader(FByteReader Values, FSchemaLoadPlan Schema, const FLoadBatch& InBatch)
+	: Types(Schema.GetMembers())
 	, Offsets(Schema.GetOffsets())
-	, InnerStructSchemas(static_cast<const FStructSchemaId*>(Schema.GetInnerSchemas()), Schema.NumInnerSchemas)
-	, InnerRangeTypes(Schema.GetInnerRangeTypes(), Schema.NumInnerRanges)
+	, InnerStructSchemas(Schema.GetInnerSchemas())
+	, InnerRangeTypes(Schema.GetInnerRangeTypes())
 	, RangeBindings(Schema.GetRangeBindings())
 	, Batch(InBatch)
 	, ByteIt(Values)
@@ -1129,7 +1166,7 @@ void LoadStruct(void* Dst, FByteReader Src, FStructSchemaId LoadId, const FLoadB
 	else if (Plan.IsMemcpy())
 	{
 		static_assert(PLATFORM_LITTLE_ENDIAN);
-		FLoadStructMemcpy MemcpyPlan = Plan.AsMemcpy();
+		FMemcpyLoadPlan MemcpyPlan = Plan.AsMemcpy();
 		Src.CheckSize(MemcpyPlan.Size);
 		FMemory::Memcpy(static_cast<uint8*>(Dst) + MemcpyPlan.Offset, Src.Peek(), MemcpyPlan.Size);
 	}
@@ -1206,13 +1243,6 @@ FMemberLoader::FMemberLoader(FStructLoadView In)
 , Batch(In.Schema.Batch)
 {}
 
-static FMemberType GetInnermostType(FRangeSchema Schema)
-{
-	FMemberType Inner = Schema.ItemType;
-	for (const FMemberType* It = Schema.NestedItemTypes; Inner.IsRange(); Inner = *It++) {}
-	return Inner;
-}
-
 FRangeLoadView FMemberLoader::GrabRange()
 {
 	const FRangeView In = Reader.GrabRange();
diff --git a/Engine/Plugins/Experimental/PlainProps/Source/Private/PlainPropsParse.cpp b/Engine/Plugins/Experimental/PlainProps/Source/Private/PlainPropsParse.cpp
--- a/Engine/Plugins/Experimental/PlainProps/Source/Private/PlainPropsParse.cpp
+++ b/Engine/Plugins/Experimental/PlainProps/Source/Private/PlainPropsParse.cpp
@@ -35,15 +35,15 @@ auto Print(FToken Token)
 	return WriteToString<64>(FStringView(StringCast<TCHAR>(Token.Str, Token.Len)));
 }
 
-inline const FParsedStructSchema& ResolveStructSchema(const FParsedSchemas& Schemas, FStructSchemaId Id)
+inline const FParsedStructSchema& ResolveStructSchema(const FParsedSchemas& Schemas, FStructId Id)
 {
-	check(Id.Idx < IntCastChecked<uint32>(Schemas.Structs.Num()));
+	check(Id.Idx < static_cast<uint32>(Schemas.Structs.Num()));
 	return Schemas.Structs[Id.Idx];
 }
 
-inline const FParsedEnumSchema& ResolveEnumSchema(const FParsedSchemas& Schemas, FEnumSchemaId Id)
+inline const FParsedEnumSchema& ResolveEnumSchema(const FParsedSchemas& Schemas, FEnumId Id)
 {
-	check(Id.Idx < IntCastChecked<uint32>(Schemas.Enums.Num()));
+	check(Id.Idx < static_cast<uint32>(Schemas.Enums.Num()));
 	return Schemas.Enums[Id.Idx];
 }
 
@@ -540,7 +540,7 @@ public:
 		FDeclarations& InTypes,
 		FParsedSchemas& InSchemas);
 
-	FBuiltStruct* ParseMembers(FStructSchemaId SchemaId);
+	FBuiltStruct* ParseMembers(FStructId SchemaId);
 
 private:
 	FDebugIds GetDebug() const { return FDebugIds(Names); }
@@ -552,9 +552,9 @@ private:
 	bool ParseRange();
 
 	bool AdvanceToNextMember();
-	bool ParseDynamicStructType(FOptionalStructSchemaId& Out);
+	FStructId ParseDynamicStructType();
 
-	[[nodiscard]] FBuiltStruct* BuildAndReset(FStructSchemaId Schema);
+	[[nodiscard]] FBuiltStruct* BuildAndReset(FStructId Schema);
 
 	[[nodiscard]] FTypedRange	ParseLeaves(const FMemberSchema& RangeSchema, FScopedTokenReader RangeTokens);
 	[[nodiscard]] FTypedRange	ParseStructs(const FMemberSchema& RangeSchema, FScopedTokenReader RangeTokens);
@@ -564,7 +564,7 @@ private:
 
 	[[nodiscard]] FBuiltStruct*	ParseMembersInternal(
 		FStructType StructType,
-		FStructSchemaId Schema,
+		FStructId Schema,
 		FScopedTokenReader InnerTokens);
 
 	[[nodiscard]] FTypedRange	ParseUnicodeLeafValueAsRange(
@@ -580,7 +580,7 @@ private:
 	FTextIndexer&							Names;
 	FDeclarations&							Types;
 	FParsedSchemas&							Schemas;
-	FOptionalStructSchemaId					UsedSuper;
+	FOptionalStructId						UsedSuper;
 	uint32									TokenIdx = 0;
 	uint32								 	MemberIdx = 0;
 	bool									HasSuper = false;
@@ -952,7 +952,7 @@ FParsedEnumSchema FEnumSchemaParser::DeclareEnumSchema(
 	TConstArrayView<uint64> ConstantValues)
 {
 	FType Type = MakeType(TypeTokens, Names);
-	FEnumSchemaId Id = Names.IndexEnum(Type);
+	FEnumId Id = Names.IndexEnum(Type);
 
 	TArray<FEnumerator, TInlineAllocator<64>> Enumerators;
 	Enumerators.Reserve(ConstantNames.Num());
@@ -1045,12 +1045,10 @@ FParsedStructSchema FStructSchemaParser::DeclareStructSchema(
 	TConstArrayView<FUtf8StringView> MemberTypes)
 {
 	const FType Type = MakeType(TypeTokens, Names);
-	const FStructSchemaId Id = Names.IndexStruct(Type);
-	const FOptionalStructSchemaId DeclaredSuper = ParsedSuper ?
-		Names.IndexStruct(MakeType(ParsedSuper.GetValue(), Names)) :
-		FOptionalStructSchemaId(NoId);
+	const FStructId Id = Names.IndexStruct(Type);
+	const FOptionalStructId DeclaredSuper = ParsedSuper ? ToOptional(Names.IndexStruct(MakeType(ParsedSuper.GetValue(), Names))) : NoId;
 
-	FOptionalStructSchemaId UsedSuper;
+	FOptionalStructId UsedSuper;
 
 	TArray<FMemberId, TInlineAllocator<64>> MemberIds;
 	TArray<FParsedMemberSchema> MemberSchemas;
@@ -1131,23 +1129,19 @@ bool FStructSchemaParser::ParseMemberSchema(FParsedMemberSchema& Out, FUtf8Strin
 	else if (TOptional<FTypeTokens> TypeTokens = TokenizeType(TypeOrLeaf))
 	{
 		FType Type = MakeType(TypeTokens.GetValue(), Names);
-		FOptionalEnumSchemaId Enum = Names.GetEnumId(Type);
+		FOptionalEnumId Enum = Names.GetEnumId(Type);
 		if (Enum)
 		{
 			// All Enum types have already been parsed and are known before parsing structs and members
 			const FParsedEnumSchema& EnumSchema = ResolveEnumSchema(Schemas, Enum.Get());
 			Out.Type = FMemberType(ELeafType::Enum, EnumSchema.Width);
-			Out.InnerSchema = Enum.Get();
+			Out.InnerSchema = FInnerId(Enum.Get());
 		}
 		else
 		{
 			// PP-TEXT: Index all Struct types in a pass before parsing struct members?
-			// FOptionalStructSchemaId Struct = Names.GetStructId(Type);
-			// check(Struct);
-			// Out.InnerSchema = Struct.Get();
-			Out.Type = FMemberType(FStructType{ EMemberKind::Struct});
-			FStructSchemaId Struct = Names.IndexStruct(Type);
-			Out.InnerSchema = Struct;
+			Out.Type = DefaultStructType;
+			Out.InnerSchema = FInnerId(Names.IndexStruct(Type));
 		}
 	}
 
@@ -1193,10 +1187,10 @@ FMemberParser::FMemberParser(
 {
 }
 
-FBuiltStruct* FMemberParser::ParseMembers(FStructSchemaId SchemaId)
+FBuiltStruct* FMemberParser::ParseMembers(FStructId Id)
 {
 	ParseAll();
-	return BuildAndReset(SchemaId);
+	return BuildAndReset(Id);
 }
 
 bool FMemberParser::PeekKind(EMemberKind& MemberKind) const
@@ -1250,7 +1244,7 @@ void AddMember(FMemberBuilder& Members, FMemberId Id, FUtf8StringView String)
 }
 
 template<typename T>
-void AddEnum(FMemberBuilder& Members, FMemberId Id, FEnumSchemaId Enum, FUtf8StringView String)
+void AddEnum(FMemberBuilder& Members, FMemberId Id, FEnumId Enum, FUtf8StringView String)
 {
 	if (TOptional<T> Value = Parse<T>(String))
 	{
@@ -1288,7 +1282,7 @@ bool FMemberParser::ParseLeaf()
 	}
 
 	const FUnpackedLeafType Leaf = MemberSchema.Type.AsLeaf();
-	const FOptionalEnumSchemaId Enum = static_cast<FOptionalEnumSchemaId>(MemberSchema.InnerSchema);
+	const FOptionalEnumId Enum = ToOptionalEnum(MemberSchema.InnerSchema);
 
 	switch (Leaf.Type)
 	{
@@ -1355,7 +1349,7 @@ bool FMemberParser::ParseStruct()
 		FToken Token = MemberTokens.GrabToken(EToken::BeginStruct);
 		if (!UsedSuper)
 		{
-			verify(ParseDynamicStructType(UsedSuper));
+			UsedSuper = ParseDynamicStructType();
 		}
 
 		FBuiltStruct* Struct = ParseMembersInternal(
@@ -1379,14 +1373,10 @@ bool FMemberParser::ParseStruct()
 
 	FToken Token = MemberTokens.GrabToken(EToken::BeginStruct);
 
-	FOptionalStructSchemaId Schema = static_cast<FOptionalStructSchemaId>(MemberSchema.InnerSchema);
-	if (StructType.IsDynamic)
-	{
-		verify(ParseDynamicStructType(Schema));
-	}
+	FStructId Id = StructType.IsDynamic ? ParseDynamicStructType() : MemberSchema.InnerSchema.Get().AsStruct();
 
-	FBuiltStruct* Struct = ParseMembersInternal(StructType, Schema.Get(), { MemberTokens.GetTokens(), Token });
-	MemberBuilder.AddStruct(MemberId, Schema.Get(), Struct);
+	FBuiltStruct* Struct = ParseMembersInternal(StructType, Id, { MemberTokens.GetTokens(), Token });
+	MemberBuilder.AddStruct(MemberId, Id, Struct);
 
 	Token = MemberTokens.GrabToken(EToken::EndStruct);
 	return true;
@@ -1446,7 +1436,7 @@ bool FMemberParser::AdvanceToNextMember()
 	return false;
 }
 
-bool FMemberParser::ParseDynamicStructType(FOptionalStructSchemaId& Out)
+FStructId FMemberParser::ParseDynamicStructType()
 {
 	FToken Token = MemberTokens.GrabToken(EToken::LeafId);
 	check(Token.Value() == GLiterals.Dynamic);
@@ -1456,20 +1446,14 @@ bool FMemberParser::ParseDynamicStructType(FOptionalStructSchemaId& Out)
 	TOptional<FTypeTokens> TypeTokens = TokenizeType(StructName);
 	check(TypeTokens);
 	FType Type = MakeType(TypeTokens.GetValue(), Names);
-	FOptionalStructSchemaId Schema = Names.GetStructId(Type); 
-	if (Schema)
-	{
-		Out = Schema;
-		return true;
-	}
-	checkf(false, TEXT("Failed to parse dynamic struct schema '%s'"), *Print(Token));
-	return false;
+	FOptionalStructId Out = Names.GetStructId(Type);
+	checkf(Out, TEXT("Failed to parse dynamic struct schema '%s'"), *Print(Token));
+	return Out.Get();
 }
 
-FBuiltStruct* FMemberParser::BuildAndReset(FStructSchemaId Schema)
+FBuiltStruct* FMemberParser::BuildAndReset(FStructId Id)
 {
-	const FStructDeclaration& StructDeclaration = Types.Get(static_cast<FStructSchemaId>(Schema));
-	return MemberBuilder.BuildAndReset(Scratch, StructDeclaration, GetDebug());
+	return MemberBuilder.BuildAndReset(Scratch, Types.Get(Id), GetDebug());
 }
 
 template<Arithmetic T>
@@ -1493,7 +1477,7 @@ template<Arithmetic T>
 template<Arithmetic T>
 [[nodiscard]] static FTypedRange ParseEnumRange(
 	FScratchAllocator& Scratch,
-	FEnumSchemaId Enum,
+	FEnumId Enum,
 	ERangeSizeType MaxSize,
 	FScopedTokenReader& RangeTokens)
 {
@@ -1511,10 +1495,10 @@ template<Arithmetic T>
 
 FTypedRange FMemberParser::ParseLeaves(const FMemberSchema& RangeSchema, FScopedTokenReader RangeTokens)
 {
-	const FOptionalEnumSchemaId Enum = static_cast<FOptionalEnumSchemaId>(RangeSchema.InnerSchema);
+	const FOptionalEnumId Enum = ToOptionalEnum(RangeSchema.InnerSchema);
 	const FMemberType InnermostType = RangeSchema.GetInnerRangeTypes().Last();
 	const FUnpackedLeafType Leaf = InnermostType.AsLeaf();
-	ERangeSizeType NumType = RangeSchema.Type.AsRange().MaxSize;
+	const ERangeSizeType NumType = RangeSchema.Type.AsRange().MaxSize;
 
 	FTypedRange TypedRange;
 	switch (Leaf.Type)
@@ -1585,7 +1569,7 @@ FTypedRange FMemberParser::ParseStructs(const FMemberSchema& RangeSchema, FScope
 	check(RangeSchema.InnerSchema || StructType.IsDynamic);
 
 	TArray<FBuiltStruct*, TInlineAllocator64<64>> Structs;
-	FOptionalStructSchemaId Schema = static_cast<FOptionalStructSchemaId>(RangeSchema.InnerSchema);
+	FOptionalStructId Schema = ToOptionalStruct(RangeSchema.InnerSchema);
 
 	if (StructType.IsDynamic)
 	{
@@ -1593,11 +1577,8 @@ FTypedRange FMemberParser::ParseStructs(const FMemberSchema& RangeSchema, FScope
 		while (RangeTokens.HasMore())
 		{
 			FToken Token = RangeTokens.GrabToken(EToken::BeginStruct);
-			FOptionalStructSchemaId DynamicSchema;
-			verify(ParseDynamicStructType(DynamicSchema));
-			
-			checkf(!Schema || Schema == DynamicSchema,
-				TEXT("Heterogeneous struct ranges have not been implemented yet"));
+			FStructId DynamicSchema = ParseDynamicStructType();
+			checkf(!Schema || Schema == DynamicSchema, TEXT("Heterogeneous struct ranges have not been implemented yet"));
 			Schema = DynamicSchema;
 
 			FBuiltStruct* Struct = ParseMembersInternal(
@@ -1704,7 +1685,7 @@ FTypedRange FMemberParser::ParseRangeInternal(const FMemberSchema& RangeSchema)
 
 FBuiltStruct* FMemberParser::ParseMembersInternal(
 	FStructType StructType,
-	FStructSchemaId Schema,
+	FStructId Schema,
 	FScopedTokenReader InnerTokens)
 {
 	const FParsedStructSchema& ParsedSchema = ResolveStructSchema(Schemas, Schema);
@@ -1849,7 +1830,7 @@ void FBatchParser::ParseObjects(TConstArrayView64<FToken> TokensView)
 		}
 
 		const FType Type = MakeType(TypeTokens.GetValue(), Names);
-		const FOptionalStructSchemaId StructId = Names.GetStructId(Type);
+		const FOptionalStructId StructId = Names.GetStructId(Type);
 
 		checkf(StructId, TEXT("Object struct schema '%s' not found"), *Print(Token));
 		const FParsedStructSchema& ParsedSchema = ResolveStructSchema(Schemas, StructId.Get());
@@ -1869,7 +1850,7 @@ void FBatchParser::ParseObjects(TConstArrayView64<FToken> TokensView)
 class FParseBindings : public IStructBindIds
 {
 public:
-	virtual FStructSchemaId GetDeclId(FStructSchemaId BindId) const override final
+	virtual FStructId GetDeclId(FStructId BindId) const override final
 	{
 		checkf(false, TEXT("All struct ids should be declared, nothing is bound with different names for text"));
 		return BindId;
@@ -1880,7 +1861,7 @@ void FBatchParser::Write(TArray64<uint8>& Out)
 {
 	FParseBindings Bindings;
 	FSchemasBuilder SchemaBuilders(Types, Names, Bindings, Scratch, ESchemaFormat::StableNames);
-	for (const TPair<FStructSchemaId, FBuiltStruct*>& Object : Objects)
+	for (const TPair<FStructId, FBuiltStruct*>& Object : Objects)
 	{
 		SchemaBuilders.NoteStructAndMembers(Object.Key, *Object.Value);
 	}
@@ -1900,4 +1881,4 @@ void ParseYamlBatch(TArray64<uint8>& OutBinary, FUtf8StringView Yaml)
 	BatchParser.Parse(OutBinary);
 }
 
-} // namespace PlainProps
+} // namespace PlainProps
\ No newline at end of file
diff --git a/Engine/Plugins/Experimental/PlainProps/Source/Private/PlainPropsPrint.cpp b/Engine/Plugins/Experimental/PlainProps/Source/Private/PlainPropsPrint.cpp
--- a/Engine/Plugins/Experimental/PlainProps/Source/Private/PlainPropsPrint.cpp
+++ b/Engine/Plugins/Experimental/PlainProps/Source/Private/PlainPropsPrint.cpp
@@ -1,6 +1,7 @@
 // Copyright Epic Games, Inc. All Rights Reserved.
 
 #include "PlainPropsPrint.h"
+#include "PlainPropsBuild.h"
 #include "PlainPropsDiff.h"
 #include "PlainPropsIndex.h"
 #include "Containers/StringConv.h"
@@ -55,63 +56,87 @@ void FIdIndexerBase::InitParameterNames()
 
 ///////////////////////////////////////////////////////////////////////////////
 
-FUtf8StringView PrintMemberSchemaImpl(
-	FUtf8Builder& Out,
-	const FIds& Ids,
-	FMemberType Type,
-	FOptionalSchemaId InnerSchema = NoId,
-	TConstArrayView<FMemberType> InnerRangeTypes = {})
+static void PrintRangeType(FUtf8Builder& Out, FRangeType Type)
 {
-	FMemberType InnermostType = InnerRangeTypes.Num() > 0 ? InnerRangeTypes.Last() : Type;
-	check(!InnermostType.IsRange());
-	if (InnerSchema)
+	Out.AppendChar('(');
+	Out.Append(ToString(Type.MaxSize));
+	Out.AppendChar(')');
+}
+
+inline void PrintStructFlags(FUtf8Builder& Out, FStructType StructType)
+{
+	if (StructType.IsDynamic)
 	{
-		check(IsStructOrEnum(InnermostType));
-		if (InnermostType.IsLeaf())
-		{
-			check(IsEnum(InnermostType));
-			Ids.AppendString(Out, static_cast<FEnumSchemaId>(InnerSchema.Get()));
-		}
-		else
-		{
-			check(InnermostType.IsStruct());
-			Ids.AppendString(Out, static_cast<FStructSchemaId>(InnerSchema.Get()));
-		}
+		Out.Append(GLiterals.Dynamic);
 	}
-	else if (InnermostType.IsLeaf())
+}
+
+void PrintSchema(FUtf8Builder& Out, const FBatchIds& Ids, FStructType StructType, FStructSchemaId Id)
+{
+	PrintStructFlags(Out, StructType);
+	Ids.AppendString(Out, Id);
+}
+
+template<typename IdsType, typename OptionalStructId>
+void PrintSchema(FUtf8Builder& Out, const IdsType& Ids, FStructType StructType, OptionalStructId Id)
+{
+	PrintStructFlags(Out, StructType);
+	if (Id)
 	{
-		check(!IsEnum(InnermostType));
-		Out.Append(ToString(InnermostType.AsLeaf()));
+		Ids.AppendString(Out, Id.Get());
 	}
-	else if (InnermostType.IsStruct())
+}
+
+template<typename IdsType, typename OptionalEnumId>
+void PrintSchema(FUtf8Builder& Out, const IdsType& Ids, FUnpackedLeafType Leaf, OptionalEnumId Id)
+{
+	if (Id)
 	{
-		check(InnermostType.AsStruct().IsDynamic);
-		Out.Append(GLiterals.Dynamic);
+		Ids.AppendString(Out, Id.Get());
 	}
+	else
+	{
+		Out.Append(ToString(Leaf));
+	}	
+}
 
-	if (Type.IsRange())
+template<typename IdsType, typename OptionalId>
+void PrintInnermostSchema(FUtf8Builder& Out, const IdsType& Ids, FMemberType InnermostType, OptionalId InnerSchema)
+{
+	if (InnermostType.IsStruct())
 	{
-		Out.AppendChar('(');
-		Out.Append(ToString(Type.AsRange().MaxSize));
-		Out.AppendChar(')');
-		for (FMemberType MemberType : InnerRangeTypes.LeftChop(1))
-		{
-			Out.AppendChar('(');
-			Out.Append(ToString(MemberType.AsRange().MaxSize));
-			Out.AppendChar(')');
-		}
+		PrintSchema(Out, Ids, InnermostType.AsStruct(), ToOptionalStruct(InnerSchema));
+	}
+	else
+	{
+		PrintSchema(Out, Ids, InnermostType.AsLeaf(), ToOptionalEnum(InnerSchema));
+	}
+}
+
+static void PrintSchema(FUtf8Builder& Out, const FBatchIds& Ids, FRangeType Type, FRangeSchema Schema)
+{
+	PrintInnermostSchema(Out, Ids, GetInnermostType(Schema), Schema.InnermostSchema);
+
+	PrintRangeType(Out, Type);
+	FMemberType Inner = Schema.ItemType;
+	for (const FMemberType* It = Schema.NestedItemTypes; Inner.IsRange(); Inner = *It++)
+	{
+		PrintRangeType(Out, Inner.AsRange());
 	}
-	return Out.ToView();
 }
 
-void PrintMemberSchema(
-	FUtf8Builder& Out,
-	const FIds& Ids, 
-	FMemberType Type,
-	FOptionalSchemaId InnerSchema,
-	TConstArrayView<FMemberType> InnerRangeTypes)
+void PrintMemberSchema(FUtf8Builder& Out, const FIds& Ids, FMemberSchema Member)
 {
-	PrintMemberSchemaImpl(Out, Ids, Type, InnerSchema, InnerRangeTypes);
+	PrintInnermostSchema(Out, Ids, Member.GetInnermostType(), Member.InnerSchema);
+
+	if (Member.Type.IsRange())
+	{
+		PrintRangeType(Out, Member.Type.AsRange());
+		for (FMemberType Inner : Member.GetInnerRangeTypes())
+		{
+			PrintRangeType(Out, Inner.AsRange());
+		}
+	}	
 }
 
 ///////////////////////////////////////////////////////////////////////////////
@@ -129,6 +154,12 @@ struct FMemberSchemaView
 	{
 		return InnerRangeTypes.Num() > 0 ? InnerRangeTypes.Last() : Type;
 	}
+
+	FRangeSchema AsRangeSchema() const
+	{
+		check(Type.IsRange());
+		return { InnerRangeTypes[0], Batch, InnerSchema, InnerRangeTypes.Num() > 1 ? &InnerRangeTypes[1] : nullptr };
+	}
 };
 
 ///////////////////////////////////////////////////////////////////////////////
@@ -233,45 +264,53 @@ private:
 class FMemberPrinter
 {
 public:
-	FMemberPrinter(FYamlBuilder& InTextBuilder, const FIds& InIds)
+	FMemberPrinter(FYamlBuilder& InTextBuilder, const FBatchIds& InIds)
 	: TextBuilder(InTextBuilder)
 	, Ids(InIds)
 	{}
 
-	void PrintMembers(const FStructView& StructView);
+	void PrintMembers(FStructView StructView);
 
 private:
-	void PrintLeaf(FMemberId Id, const FLeafView& LeafView);
-	void PrintStruct(FOptionalMemberId Id, FStructType StructType, const FStructView& StructView);
-	void PrintRange(FMemberId Id, ERangeSizeType NumType, const FRangeView& RangeView);
+	void PrintLeaf(FMemberId Id, FLeafView LeafView);
+	void PrintStruct(FOptionalMemberId Id, FStructType StructType, FStructView StructView);
+	void PrintRange(FMemberId Id, FRangeType RangeType, const FRangeView& RangeView);
 
 	void PrintLeaves(FUnpackedLeafType Leaf, const FLeafRangeView& LeafRange);
 	void PrintStructs(FStructType StructType, const FStructRangeView& StructRange);
-	void PrintRanges(ERangeSizeType NumType, const FNestedRangeView& NestedRange);
+	void PrintRanges(FRangeType RangeType, const FNestedRangeView& NestedRange);
 
-	void PrintMembersInternal(FStructType StructType, const FStructView& StructView, const FStructSchema& Schema);
-	void PrintRangeInternal(ERangeSizeType NumType, const FRangeView& RangeView);
+	void PrintMembersInternal(FStructType StructType, FStructView StructView, const FStructSchema& Schema);
+	void PrintRangeInternal(FRangeType RangeType, const FRangeView& RangeView);
 
 	bool IsUnicodeString(const FRangeView& RangeView);
-	void PrintUnicodeLeafValue(const FLeafView& LeafView);
-	void PrintUnicodeRangeAsLeaf(FOptionalMemberId Id, ERangeSizeType NumType, const FRangeView& RangeView);
+	void PrintUnicodeLeafValue(FLeafView LeafView);
+	void PrintUnicodeRangeAsLeaf(FOptionalMemberId Id, FRangeType RangeType, const FRangeView& RangeView);
+	
+	template<typename MemberType, typename SchemaType>
+	void PrintSchemaComment(MemberType Type, SchemaType Schema)
+	{
+		if constexpr (PrintWithComments)
+		{
+			PrintSchema(/* out */ Tmp, Ids, Type, Schema);
+			TextBuilder.AddComment(Tmp);
+			Tmp.Reset();
+		}
+	}	
 
 	FYamlBuilder& TextBuilder;
-	const FIds& Ids;
+	const FBatchIds& Ids;
+	TUtf8StringBuilder<256> Tmp;
 };
 
 ///////////////////////////////////////////////////////////////////////////////
 
-void PrintYamlBatch(
-	FUtf8Builder& Out,
-	const FIds& Ids,
-	FReadBatchId Batch,
-	TConstArrayView<FStructView> Objects)
+void PrintYamlBatch(FUtf8Builder& Out, const FBatchIds& Ids, TConstArrayView<FStructView> Objects)
 {
 	FYamlBuilder YamlBuilder(Out);
 	FBatchPrinter Printer(YamlBuilder, Ids);
 	YamlBuilder.BeginDocument();
-	Printer.PrintSchemas(Batch);
+	Printer.PrintSchemas();
 	Printer.PrintObjects(Objects);
 	YamlBuilder.EndDocument();
 }
@@ -557,32 +596,25 @@ void FYamlBuilder::AppendIdentifier(FUtf8StringView Id)
 
 ///////////////////////////////////////////////////////////////////////////////
 
-
-FBatchPrinter::FBatchPrinter(FYamlBuilder& InTextBuilder, const FIds& InIds)
+FBatchPrinter::FBatchPrinter(FYamlBuilder& InTextBuilder, const FBatchIds& InIds)
 : TextBuilder(InTextBuilder)
 , Ids(InIds)
-{
-}
+{}
 
 FBatchPrinter::~FBatchPrinter()
 {}
 
-void FBatchPrinter::PrintSchemas(FReadBatchId Batch)
-{
-	PrintSchemas(GetReadSchemas(Batch), Batch);
-}
-
-void FBatchPrinter::PrintSchemas(const FSchemaBatch& Batch, FReadBatchId BatchId)
+void FBatchPrinter::PrintSchemas()
 {
 	TextBuilder.BeginRange(GLiterals.Structs);
-	for (const FStructSchema& Struct : GetStructSchemas(Batch))
+	for (const FStructSchema& Struct : GetStructSchemas(Ids.GetSchemas()))
 	{
-		PrintStructSchema(Struct, BatchId);
+		PrintStructSchema(Struct, Ids.GetBatchId());
 	}
 	TextBuilder.EndRange();
 
 	TextBuilder.BeginRange(GLiterals.Enums);
-	for (const FEnumSchema& EnumSchema : GetEnumSchemas(Batch))
+	for (const FEnumSchema& EnumSchema : GetEnumSchemas(Ids.GetSchemas()))
 	{
 		PrintEnumSchema(EnumSchema);
 	}
@@ -599,12 +631,14 @@ void FBatchPrinter::PrintObjects(TConstArrayView<FStructView> Objects)
 	TextBuilder.EndRange();
 }
 
-static FUtf8StringView PrintMemberSchema(
-	FUtf8Builder& Out,
-	const FIds& Ids,
-	const FMemberSchemaView& Schema)
+static void PrintMemberSchema(FUtf8Builder& Out, const FBatchIds& Ids, const FMemberSchemaView& Schema)
 {
-	return PrintMemberSchemaImpl(Out, Ids, Schema.Type, Schema.InnerSchema, Schema.InnerRangeTypes);
+	switch (Schema.Type.GetKind())
+	{
+	case EMemberKind::Leaf:		PrintSchema(Out, Ids, Schema.Type.AsLeaf(), ToOptionalEnum(Schema.InnerSchema)); break;
+	case EMemberKind::Range:	PrintSchema(Out, Ids, Schema.Type.AsRange(), Schema.AsRangeSchema()); break;
+	case EMemberKind::Struct:	PrintSchema(Out, Ids, Schema.Type.AsStruct(), ToOptionalStruct(Schema.InnerSchema)); break;
+	}
 }
 
 template <int32 BufferSize>
@@ -613,7 +647,7 @@ class FPrintId
 	TUtf8StringBuilder<BufferSize>		Buffer;
 public:
 	template <typename T>
-	FPrintId(const FIds& Ids, T Id)		{ Ids.AppendString(Buffer, Id); }
+	FPrintId(const FBatchIds& Ids, T Id) { Ids.AppendString(Buffer, Id); }
 	FUtf8StringView operator*() const	{ return Buffer.ToView(); }
 };
 
@@ -630,11 +664,16 @@ void FBatchPrinter::PrintStructSchema(const FStructSchema& Struct, FReadBatchId
 	}
 
 	TextBuilder.BeginRange(GLiterals.Members);
+	TUtf8StringBuilder<256> Buf;
 	while (Reader.HasMore())
 	{
-		TUtf8StringBuilder<256> Buf;
-		FOptionalMemberId MemberId = Reader.PeekName();
-		TextBuilder.AddLeaf(*FPrintId<128>(Ids, MemberId), PrintMemberSchema(Buf, Ids, Reader.GrabMember()));
+		Ids.AppendString(Buf, Reader.PeekName());
+		TextBuilder.AddLeafId(Buf);
+		Buf.Reset();
+
+		PrintMemberSchema(Buf, Ids, Reader.GrabMember());
+		TextBuilder.AddLeafValue(FUtf8StringView(Buf));
+		Buf.Reset();
 	}
 	TextBuilder.EndRange();
 
@@ -704,14 +743,14 @@ void FBatchPrinter::PrintEnumConstants(
 
 ///////////////////////////////////////////////////////////////////////////////
 
-void FMemberPrinter::PrintMembers(const FStructView& StructView)
+void FMemberPrinter::PrintMembers(FStructView StructView)
 {
 	const FStructSchema& Schema = StructView.Schema.Resolve();
 	TextBuilder.BeginStruct(*FPrintId<128>(Ids, Schema.Type));
 	PrintMembersInternal({ EMemberKind::Struct }, StructView, Schema);
 }
 
-void FMemberPrinter::PrintLeaf(FMemberId Id, const FLeafView& LeafView)
+void FMemberPrinter::PrintLeaf(FMemberId Id, FLeafView LeafView)
 {
 	TextBuilder.AddLeafId(*FPrintId<128>(Ids, Id));
 
@@ -766,30 +805,25 @@ void FMemberPrinter::PrintLeaf(FMemberId Id, const FLeafView& LeafView)
 		break;
 	}
 
-	if constexpr(PrintWithComments)
-	{
-		TUtf8StringBuilder<256> Buf;
-		PrintMemberSchemaImpl(Buf, Ids, FMemberType(LeafView.Leaf.Pack()), static_cast<FOptionalSchemaId>(LeafView.Enum));
-		TextBuilder.AddComment(Buf);
-	}
+	PrintSchemaComment(LeafView.Leaf, LeafView.Enum);
 }
 
-void FMemberPrinter::PrintStruct(FOptionalMemberId MemberId, FStructType StructType, const FStructView& StructView)
+void FMemberPrinter::PrintStruct(FOptionalMemberId MemberId, FStructType StructType, FStructView StructView)
 {
 	TextBuilder.BeginStruct(*FPrintId<128>(Ids, MemberId));
 	PrintMembersInternal(StructType, StructView, StructView.Schema.Resolve());
 }
 
-void FMemberPrinter::PrintRange(FMemberId Id, ERangeSizeType NumType, const FRangeView& RangeView)
+void FMemberPrinter::PrintRange(FMemberId Id, FRangeType RangeType, const FRangeView& RangeView)
 {
 	if (IsUnicodeString(RangeView))
 	{
-		PrintUnicodeRangeAsLeaf(Id, NumType, RangeView);
+		PrintUnicodeRangeAsLeaf(Id, RangeType, RangeView);
 	}
 	else
 	{
 		TextBuilder.BeginRange(*FPrintId<128>(Ids, Id));
-		PrintRangeInternal(NumType, RangeView);
+		PrintRangeInternal(RangeType, RangeView);
 	}
 }
 
@@ -856,34 +890,30 @@ void FMemberPrinter::PrintStructs(FStructType StructType, const FStructRangeView
 	}
 }
 
-void FMemberPrinter::PrintRanges(ERangeSizeType NumType, const FNestedRangeView& NestedRange)
+void FMemberPrinter::PrintRanges(FRangeType RangeType, const FNestedRangeView& NestedRange)
 {
 	for (FRangeView RangeView : NestedRange)
 	{
 		if (IsUnicodeString(RangeView))
 		{
-			PrintUnicodeRangeAsLeaf(NoId, NumType, RangeView);
+			PrintUnicodeRangeAsLeaf(NoId, RangeType, RangeView);
 		}
 		else
 		{
 			TextBuilder.BeginRange();
-			PrintRangeInternal(NumType, RangeView);
+			PrintRangeInternal(RangeType, RangeView);
 		}
 	}
 }
 
-void FMemberPrinter::PrintMembersInternal(FStructType StructType, const FStructView& StructView, const FStructSchema& Schema)
+void FMemberPrinter::PrintMembersInternal(FStructType StructType, FStructView StructView, const FStructSchema& Schema)
 {
 	FMemberReader It(Schema, StructView.Values, StructView.Schema.Batch);
 
 	const bool HasMembers = StructType.IsDynamic || It.HasMore();
-	if constexpr(PrintWithComments)
+	if (HasMembers)
 	{
-		if (HasMembers)
-		{
-			TUtf8StringBuilder<256> Buf;
-			TextBuilder.AddComment(PrintMemberSchemaImpl(Buf, Ids, FMemberType(StructType), StructView.Schema.Id));
-		}
+		PrintSchemaComment(StructType, StructView.Schema.Id);
 	}
 
 	if (StructType.IsDynamic)
@@ -903,53 +933,24 @@ void FMemberPrinter::PrintMembersInternal(FStructType StructType, const FStructV
 				PrintStruct(Id, Type.AsStruct(), It.GrabStruct());
 				break;
 			case EMemberKind::Range:
-				PrintRange(Id.Get(), Type.AsRange().MaxSize, It.GrabRange());
+				PrintRange(Id.Get(), Type.AsRange(), It.GrabRange());
 				break;
 		}
 	}
 	TextBuilder.EndStruct();
 
-	if constexpr(PrintWithComments)
-	{
-		if (!HasMembers)
-		{
-			TUtf8StringBuilder<256> Buf;
-			TextBuilder.AddComment(PrintMemberSchemaImpl(Buf, Ids, FMemberType(StructType), StructView.Schema.Id));
-		}
-	}
-}
-
-static TArray<FMemberType, TInlineAllocator<16>> GetInnerRangeTypes(const FRangeSchema& Schema)
-{
-	TArray<FMemberType, TInlineAllocator<16>> InnerRangeTypes;
-	InnerRangeTypes.Add(Schema.ItemType);
-	if (Schema.ItemType.IsRange())
+	if (!HasMembers)
 	{
-		const FMemberType* It = Schema.NestedItemTypes;
-		do { InnerRangeTypes.Add(*It); } while ((It++)->IsRange());
+		PrintSchemaComment(StructType, StructView.Schema.Id);
 	}
-	return InnerRangeTypes;
-}
-
-static void AddRangeSchemaComment(
-	FYamlBuilder& TextBuilder,
-	const FIds& Ids,
-	FMemberType Type,
-	const FRangeSchema& Schema)
-{
-	TUtf8StringBuilder<256> Comment;
-	TextBuilder.AddComment(PrintMemberSchemaImpl(Comment, Ids, Type, Schema.InnermostSchema, GetInnerRangeTypes(Schema)));
 }
 
-void FMemberPrinter::PrintRangeInternal(ERangeSizeType NumType, const FRangeView& RangeView)
+void FMemberPrinter::PrintRangeInternal(FRangeType RangeType, const FRangeView& RangeView)
 {
 	const FRangeSchema& Schema = RangeView.Schema;
-	if constexpr(PrintWithComments)
+	if (RangeView.Num() > 0)
 	{
-		if (RangeView.Num() > 0)
-		{
-			AddRangeSchemaComment(TextBuilder, Ids, FMemberType(NumType), Schema);
-		}
+		PrintSchemaComment(RangeType, Schema);
 	}
 
 	switch (Schema.ItemType.GetKind())
@@ -961,17 +962,14 @@ void FMemberPrinter::PrintRangeInternal(ERangeSizeType NumType, const FRangeView
 			PrintStructs(Schema.ItemType.AsStruct(), RangeView.AsStructs());
 			break;
 		case EMemberKind::Range:
-			PrintRanges(Schema.ItemType.AsRange().MaxSize, RangeView.AsRanges());
+			PrintRanges(Schema.ItemType.AsRange(), RangeView.AsRanges());
 			break;
 	}
 	TextBuilder.EndRange();
 
-	if constexpr(PrintWithComments)
+	if (RangeView.Num() == 0)
 	{
-		if (RangeView.Num() == 0)
-		{
-			AddRangeSchemaComment(TextBuilder, Ids, FMemberType(NumType), Schema);
-		}
+		PrintSchemaComment(RangeType, Schema);
 	}
 }
 
@@ -999,7 +997,7 @@ static void AddUnicodeLeaf(FYamlBuilder& TextBuilder, LeafCharType Char)
 	}
 }
 
-void FMemberPrinter::PrintUnicodeLeafValue(const FLeafView& LeafView)
+void FMemberPrinter::PrintUnicodeLeafValue(FLeafView LeafView)
 {
 	switch (LeafView.Leaf.Width)
 	{
@@ -1027,7 +1025,7 @@ static void AddUnicodeRangeLeaf(FYamlBuilder& TextBuilder, FOptionalMemberId Id,
 	}
 }
 
-void FMemberPrinter::PrintUnicodeRangeAsLeaf(FOptionalMemberId Id, ERangeSizeType NumType, const FRangeView& RangeView)
+void FMemberPrinter::PrintUnicodeRangeAsLeaf(FOptionalMemberId Id, FRangeType RangeType, const FRangeView& RangeView)
 {
 	check(IsUnicodeString(RangeView));
 
@@ -1046,21 +1044,18 @@ void FMemberPrinter::PrintUnicodeRangeAsLeaf(FOptionalMemberId Id, ERangeSizeTyp
 		case ELeafWidth::B32:	AddUnicodeRangeLeaf<UTF32CHAR>(TextBuilder, Id, LeafRange.AsUtf32());	break;
 		case ELeafWidth::B64:	check(false);															break;
 	};
-
-	if constexpr(PrintWithComments)
-	{
-		AddRangeSchemaComment(TextBuilder, Ids, FMemberType(NumType), RangeView.Schema);
-	}
+	
+	PrintSchemaComment(FRangeType(RangeType), RangeView.Schema);
 }
 
 ///////////////////////////////////////////////////////////////////////////////
 
-void FIds::AppendString(FUtf8Builder& Out, FMemberId Name) const
+void FIdsBase::AppendString(FUtf8Builder& Out, FMemberId Name) const
 {
 	AppendString(Out, Name.Id);
 }
 
-void FIds::AppendString(FUtf8Builder& Out, FOptionalMemberId Name) const
+void FIdsBase::AppendString(FUtf8Builder& Out, FOptionalMemberId Name) const
 {
 	if (Name)
 	{
@@ -1072,7 +1067,7 @@ void FIds::AppendString(FUtf8Builder& Out, FOptionalMemberId Name) const
 	}
 }
 
-void FIds::AppendString(FUtf8Builder& Out, FScopeId Scope) const
+void FIdsBase::AppendString(FUtf8Builder& Out, FScopeId Scope) const
 {
 	if (Scope.IsFlat())
 	{
@@ -1087,7 +1082,7 @@ void FIds::AppendString(FUtf8Builder& Out, FScopeId Scope) const
 	}
 }
 
-void FIds::AppendString(FUtf8Builder& Out, FTypenameId Typename) const
+void FIdsBase::AppendString(FUtf8Builder& Out, FTypenameId Typename) const
 {
 	if (Typename.IsConcrete())
 	{
@@ -1117,7 +1112,7 @@ void FIds::AppendString(FUtf8Builder& Out, FTypenameId Typename) const
 	}
 }
 
-void FIds::AppendString(FUtf8Builder& Out, FType Type) const
+void FIdsBase::AppendString(FUtf8Builder& Out, FType Type) const
 {
 	if (Type.Scope)
 	{
@@ -1127,12 +1122,22 @@ void FIds::AppendString(FUtf8Builder& Out, FType Type) const
 	AppendString(Out, Type.Name);
 }
 
-void FIds::AppendString(FUtf8Builder& Out, FEnumSchemaId Name) const
+void FIds::AppendString(FUtf8Builder& Out, FEnumId Name) const
+{
+	AppendString(Out, Resolve(Name));
+}
+
+void FIds::AppendString(FUtf8Builder& Out, FStructId Name) const
+{
+	AppendString(Out, Resolve(Name));
+}
+
+void FBatchIds::AppendString(FUtf8Builder& Out, FEnumSchemaId Name) const
 {
 	AppendString(Out, Resolve(Name));
 }
 
-void FIds::AppendString(FUtf8Builder& Out, FStructSchemaId Name) const
+void FBatchIds::AppendString(FUtf8Builder& Out, FStructSchemaId Name) const
 {
 	AppendString(Out, Resolve(Name));
 }
@@ -1228,7 +1233,7 @@ FString FDebugIds::Print(FType Type) const
 	return FString(FStringView(StringCast<TCHAR>(Out.GetData(), Out.Len())));
 }
 
-FString FDebugIds::Print(FEnumSchemaId Name) const
+FString FDebugIds::Print(FEnumId Name) const
 {
 	TUtf8StringBuilder<128> Out;
 	if (Name.Idx < Ids.NumEnums())
@@ -1242,7 +1247,7 @@ FString FDebugIds::Print(FEnumSchemaId Name) const
 	return FString(FStringView(StringCast<TCHAR>(Out.GetData(), Out.Len())));
 }
 
-FString FDebugIds::Print(FStructSchemaId Name) const
+FString FDebugIds::Print(FStructId Name) const
 {
 	TUtf8StringBuilder<128> Out;
 	if (Name.Idx < Ids.NumStructs())
@@ -1258,7 +1263,7 @@ FString FDebugIds::Print(FStructSchemaId Name) const
 
 ///////////////////////////////////////////////////////////////////////////////
 
-void PrintDiff(FUtf8Builder& Out, const FDiffPath& Diff, const FIds& Ids)
+void PrintDiff(FUtf8Builder& Out, const FIds& Ids, const FDiffPath& Diff)
 {
 	check(Diff.Num());
 
@@ -1280,7 +1285,7 @@ void PrintDiff(FUtf8Builder& Out, const FDiffPath& Diff, const FIds& Ids)
 		{
 			Ids.AppendString(Out, FTypenameId(Node.Meta.Range.GetBindName()));	
 		}
-		else if (FOptionalEnumSchemaId Enum = Node.Meta.Leaf)
+		else if (FOptionalEnumId Enum = Node.Meta.Leaf)
 		{
 			Ids.AppendString(Out, Ids.Resolve(Enum.Get()).Name);
 		}
diff --git a/Engine/Plugins/Experimental/PlainProps/Source/Private/PlainPropsRead.cpp b/Engine/Plugins/Experimental/PlainProps/Source/Private/PlainPropsRead.cpp
--- a/Engine/Plugins/Experimental/PlainProps/Source/Private/PlainPropsRead.cpp
+++ b/Engine/Plugins/Experimental/PlainProps/Source/Private/PlainPropsRead.cpp
@@ -415,4 +415,17 @@ FFlatMemberReader::FFlatMemberReader(FStructView Struct)
 	It = &Lineage.Last();
 }
 
+//////////////////////////////////////////////////////////////////////////
+
+FBatchIds::FBatchIds(FReadBatchId Batch) : Schemas(GReadSchemas.Get(Batch)), BatchId(Batch) {}
+uint32 FBatchIds::NumEnums() const { return Schemas.NumSchemas - Schemas.NumStructSchemas; }
+uint32 FBatchIds::NumStructs() const { return Schemas.NumStructSchemas; }
+FType FBatchIds::Resolve(FEnumSchemaId Id) const { return ResolveEnumSchema(Schemas, Id).Type; }
+FType FBatchIds::Resolve(FStructSchemaId Id) const	{ return ResolveStructSchema(Schemas, Id).Type; }
+
+uint32 FStableBatchIds::NumNestedScopes() const { return Schemas.NumNestedScopes; }
+uint32 FStableBatchIds::NumParametricTypes() const { return Schemas.NumParametricTypes; }
+FNestedScope FStableBatchIds::Resolve(FNestedScopeId Id) const { return ResolveNestedScope(Schemas, Id); }
+FParametricTypeView FStableBatchIds::Resolve(FParametricTypeId Id) const { return ResolveParametricType(Schemas, Id); }
+
 } // namespace PlainProps
diff --git a/Engine/Plugins/Experimental/PlainProps/Source/Private/PlainPropsSave.cpp b/Engine/Plugins/Experimental/PlainProps/Source/Private/PlainPropsSave.cpp
--- a/Engine/Plugins/Experimental/PlainProps/Source/Private/PlainPropsSave.cpp
+++ b/Engine/Plugins/Experimental/PlainProps/Source/Private/PlainPropsSave.cpp
@@ -124,10 +124,10 @@ struct TNonLeafRangeSaver
 	}
 };
 
-struct FDefaultStruct { FStructSchemaId Id; const void* Struct; };
+struct FDefaultStruct { FStructId Id; const void* Struct; };
 
 using FInternalNestedRangeSaver = TNonLeafRangeSaver<FBuiltRange*, FRangeMemberBinding>;
-using FInternalStructRangeSaver = TNonLeafRangeSaver<FBuiltStruct*, FStructSchemaId>;
+using FInternalStructRangeSaver = TNonLeafRangeSaver<FBuiltStruct*, FStructId>;
 using FInternalStructRangeDeltaSaver = TNonLeafRangeSaver<FBuiltStruct*, FDefaultStruct>;
 
 //////////////////////////////////////////////////////////////////////////
@@ -199,7 +199,7 @@ FBuiltRange* SaveLeafRange(const void* Range, const ILeafRangeBinding& Binding,
 	return Allocator.GetAllocatedRange();
 }
 
-inline const FStructDeclaration& GetDeclaration(const FSaveContext& Ctx, FStructSchemaId BindId)
+inline const FStructDeclaration& GetDeclaration(const FSaveContext& Ctx, FStructId BindId)
 {
 	// BindId is usually DeclId, unless it's been type-erased/lowered
 	if (const FStructDeclaration* Out = Ctx.Declarations.Find(BindId))
@@ -215,7 +215,7 @@ inline const FStructDeclaration& GetDeclaration(const FSaveContext& Ctx, FStruct
 	return Ctx.Declarations.Get(Ctx.Schemas.GetStruct(BindId).DeclId);
 }
 
-[[nodiscard]] static FBuiltRange* SaveStructRange(const void* Range, const IItemRangeBinding& ItemBinding, const FSaveContext& Ctx, FStructSchemaId BindId)
+[[nodiscard]] static FBuiltRange* SaveStructRange(const void* Range, const IItemRangeBinding& ItemBinding, const FSaveContext& Ctx, FStructId BindId)
 {
 	if (Ctx.Defaults && GetDeclaration(Ctx, BindId).Occupancy == EMemberPresence::AllowSparse)
 	{
@@ -243,7 +243,7 @@ inline const FStructDeclaration& GetDeclaration(const FSaveContext& Ctx, FStruct
 	{
 	case EMemberKind::Leaf:		return SaveLeafRange(Range, ItemBinding, Ctx, GetItemWidth(InnerType.AsLeaf()));
 	case EMemberKind::Range:	return SaveNonLeafRange<FInternalNestedRangeSaver>(Range, ItemBinding, Ctx, GetInnerRange(Member));
-	case EMemberKind::Struct:	return SaveStructRange(Range, ItemBinding, Ctx, static_cast<FStructSchemaId>(Member.InnermostSchema.Get()));
+	case EMemberKind::Struct:	return SaveStructRange(Range, ItemBinding, Ctx, Member.InnermostSchema.Get().AsStruct());
 	}
 
 	check(false);
@@ -257,7 +257,7 @@ inline const FStructDeclaration& GetDeclaration(const FSaveContext& Ctx, FStruct
 	return SaveRange(Range, Member, Ctx);
 }
 
-[[nodiscard]] static FBuiltStruct* SaveRangeItem(const uint8* Struct, FStructSchemaId Id, const FSaveContext& Ctx)
+[[nodiscard]] static FBuiltStruct* SaveRangeItem(const uint8* Struct, FStructId Id, const FSaveContext& Ctx)
 {
 	return SaveStruct(Struct, Id, Ctx);
 }
@@ -319,7 +319,7 @@ inline void SaveAllMembers(FMemberBuilder& Out, const void* Struct, const FSchem
 	FMemberVisitor It(Schema);
 	if (Declaration.Super)
 	{
-		FStructSchemaId SuperId = It.GrabSuper();
+		FStructId SuperId = It.GrabSuper();
 		checkSlow(SuperId == Declaration.Super);
 		const FSchemaBinding& SuperSchema = Ctx.Schemas.GetStruct(SuperId);
 		const FStructDeclaration& SuperDecl = Ctx.Declarations.Get(SuperId);
@@ -339,7 +339,7 @@ inline void SaveAllMembers(FMemberBuilder& Out, const void* Struct, const FSchem
 	checkSlow(!It.HasMore());
 }
 
-FBuiltStruct* SaveStruct(const void* Struct, FStructSchemaId BindId, const FSaveContext& Ctx)
+FBuiltStruct* SaveStruct(const void* Struct, FStructId BindId, const FSaveContext& Ctx)
 {
 	const FStructDeclaration* Declaration = nullptr;
 	FMemberBuilder Out;
@@ -374,14 +374,14 @@ static bool DiffItem(const uint8* A, const uint8* B, const FSaveContext& Ctx, FR
 	{
 	case EMemberKind::Leaf:		return DiffItemRange(A, B, ItemBinding, Ctx, SizeOf(GetItemWidth(InnerType.AsLeaf())));
 	case EMemberKind::Range:	return DiffItemRange(A, B, ItemBinding, Ctx, GetInnerRange(Member));
-	case EMemberKind::Struct:	return DiffItemRange(A, B, ItemBinding, Ctx, static_cast<FStructSchemaId>(Member.InnermostSchema.Get()));
+	case EMemberKind::Struct:	return DiffItemRange(A, B, ItemBinding, Ctx, Member.InnermostSchema.Get().AsStruct());
 	}
 
 	check(false);
 	return false;
 }
 
-static bool DiffItem(const uint8* A, const uint8* B, const FSaveContext& Ctx, FStructSchemaId Id)
+static bool DiffItem(const uint8* A, const uint8* B, const FSaveContext& Ctx, FStructId Id)
 {
 	if (const ICustomBinding* Custom = Ctx.Customs.FindStruct(Id))
 	{
@@ -442,7 +442,7 @@ static const FStructDeclaration& SaveSchemaBoundStructDelta(FMemberBuilder& Out,
 		FMemberVisitor It(Schema);
 		if (Declaration.Super)
 		{
-			FStructSchemaId SuperId = It.GrabSuper();
+			FStructId SuperId = It.GrabSuper();
 			checkSlow(SuperId == Declaration.Super);
 			const FSchemaBinding& SuperSchema = Ctx.Schemas.GetStruct(SuperId);
 			const FStructDeclaration& SuperDecl = SaveSchemaBoundStructDelta(Out, Struct, Default, SuperSchema, Ctx);
@@ -468,7 +468,7 @@ static const FStructDeclaration& SaveSchemaBoundStructDelta(FMemberBuilder& Out,
 	return Declaration;
 }
 
-FBuiltStruct* SaveStructDelta(const void* Struct, const void* Default, FStructSchemaId BindId, const FSaveContext& Ctx)
+FBuiltStruct* SaveStructDelta(const void* Struct, const void* Default, FStructId BindId, const FSaveContext& Ctx)
 {
 	FMemberBuilder Out;
 	const FStructDeclaration* Declaration = nullptr;
@@ -487,7 +487,7 @@ FBuiltStruct* SaveStructDelta(const void* Struct, const void* Default, FStructSc
 	return Out.BuildAndReset(Ctx.Scratch, *Declaration, Ctx.Declarations.GetDebug());
 }
 
-FBuiltStruct* SaveStructDeltaIfDiff(const void* Struct, const void* Default, FStructSchemaId BindId, const FSaveContext& Ctx)
+FBuiltStruct* SaveStructDeltaIfDiff(const void* Struct, const void* Default, FStructId BindId, const FSaveContext& Ctx)
 {
 	FMemberBuilder Out;
 	const FStructDeclaration* Declaration = nullptr;
diff --git a/Engine/Plugins/Experimental/PlainProps/Source/Private/PlainPropsWrite.cpp b/Engine/Plugins/Experimental/PlainProps/Source/Private/PlainPropsWrite.cpp
--- a/Engine/Plugins/Experimental/PlainProps/Source/Private/PlainPropsWrite.cpp
+++ b/Engine/Plugins/Experimental/PlainProps/Source/Private/PlainPropsWrite.cpp
@@ -47,9 +47,9 @@ struct FWriteIds
 	FConcreteTypenameId					Remap(FConcreteTypenameId Old) const	{ return { Remap(Old.Id) }; }	
 	FParametricTypeId					Remap(FParametricTypeId Old) const		{ return ParametricTypes[Old.Idx].Get(); }
 	FTypenameId							Remap(FTypenameId Old) const			{ return Old.IsConcrete() ? FTypenameId(Remap(Old.AsConcrete())) : FTypenameId(Remap(Old.AsParametric())); }
-	FType								Remap(FType Old) const				{ return { Remap(Old.Scope), Remap(Old.Name) }; }
-	FStructSchemaId						RemapStruct(FSchemaId Old) const;
-	FEnumSchemaId						RemapEnum(FSchemaId Old) const			{ return Enums[Old.Idx].Get(); }
+	FType								Remap(FType Old) const					{ return { Remap(Old.Scope), Remap(Old.Name) }; }
+	FStructSchemaId						RemapStruct(FStructId Old) const;
+	FEnumSchemaId						RemapEnum(FEnumId Old) const			{ return Enums[Old.Idx].Get(); }
 
 	template<typename T>
 	TOptionalId<T>						Remap(TOptionalId<T> Old) const			{ return Old ? ToOptional(Remap(Old.Get())) : Old; }
@@ -290,16 +290,15 @@ FWriteIds::FWriteIds(const FIdIndexerBase& Ids, const IStructBindIds& InBindIds,
 	check(HasStableNames() == (Format == ESchemaFormat::StableNames));
 }
 
-FStructSchemaId
-FWriteIds::RemapStruct(FSchemaId OldBindId) const
+FStructSchemaId FWriteIds::RemapStruct(FStructId OldBindId) const
 {
-	if (FOptionalStructSchemaId NewDeclId = Structs[OldBindId.Idx])
+	if (FOptionalStructSchemaId WriteDeclId = Structs[OldBindId.Idx])
 	{
-		return NewDeclId.Get();
+		return WriteDeclId.Get();
 	}
 
 	// Could optimize by caching Structs[OldBindId.Idx] here
-	FStructSchemaId OldDeclId = BindIds.GetDeclId(static_cast<FStructSchemaId>(OldBindId));
+	FStructId OldDeclId = BindIds.GetDeclId(OldBindId);
 	return Structs[OldDeclId.Idx].Get();
 }
 
@@ -328,17 +327,17 @@ static TArray<FMemberType> GetInnerRangeTypes(const FBuiltStructSchema& Struct)
 
 static FOptionalSchemaId GetStaticInnerSchema(const FMemberSchema& Schema, const FWriteIds& NewIds)
 {
-	if (FOptionalSchemaId InnerSchema = Schema.InnerSchema)
+	if (FOptionalInnerId InnerSchema = Schema.InnerSchema)
 	{
 		FMemberType InnermostType = Schema.GetInnermostType();
 		check(IsStructOrEnum(InnermostType));
 		if (InnermostType.IsLeaf())
 		{
-			return { NewIds.RemapEnum(InnerSchema.Get()) };
+			return { NewIds.RemapEnum(InnerSchema.Get().AsEnum()) };
 		}
 		else if (!InnermostType.AsStruct().IsDynamic)
 		{
-			return { NewIds.RemapStruct(InnerSchema.Get()) };
+			return { NewIds.RemapStruct(InnerSchema.Get().AsStruct()) };
 		}
 	}
 	return {};
@@ -373,7 +372,7 @@ static TArray<IdType> RemapIds(const FWriteIds& NewIds, const TArray<IdType>& Na
 	return Out;
 }
 
-static ESuper GetInheritance(FOptionalStructSchemaId Super, TConstArrayView<const FMemberSchema*> Members)
+static ESuper GetInheritance(FOptionalStructId Super, TConstArrayView<const FMemberSchema*> Members)
 {
 	if (Super)
 	{
@@ -382,7 +381,7 @@ static ESuper GetInheritance(FOptionalStructSchemaId Super, TConstArrayView<cons
 			return ESuper::Unused;
 		}
 
-		if (Super != static_cast<FOptionalStructSchemaId>(Members[0]->InnerSchema))
+		if (Super != ToOptionalStruct(Members[0]->InnerSchema))
 		{
 			return ESuper::Used;
 		}
@@ -638,7 +637,7 @@ public:
 	, Debug(InDebug)
 	{}
 
-	FStructSchemaId WriteMembers(FStructSchemaId BuiltId, const FBuiltStruct& Struct)
+	FStructSchemaId WriteMembers(FStructId BuiltId, const FBuiltStruct& Struct)
 	{
 		FStructSchemaId WriteId = NewIds.RemapStruct(BuiltId);
 		const FBuiltStructSchema& Schema = Schemas.Structs[WriteId.Idx];
@@ -703,15 +702,15 @@ private:
 	{
 		if (InnermostType.IsStruct() && InnermostType.AsStruct().IsDynamic)
 		{
-			FSchemaId NewId = NewIds.RemapStruct(Schema.InnerSchema.Get());
-			Bytes.Append(reinterpret_cast<uint8*>(&NewId.Idx), sizeof(NewId.Idx));
+			FSchemaId WriteId = NewIds.RemapStruct(Schema.InnerSchema.Get().AsStruct());
+			Bytes.Append(reinterpret_cast<uint8*>(&WriteId.Idx), sizeof(WriteId.Idx));
 		}
 
 		switch (Schema.Type.GetKind())
 		{
 		case EMemberKind::Leaf:		WriteLeaf(Schema.Type.AsLeaf(), Value.Leaf); break;
 		case EMemberKind::Range:	WriteRange(Schema.Type.AsRange().MaxSize, Schema.GetInnerRangeTypes(), Schema.InnerSchema, Value.Range); break;
-		case EMemberKind::Struct:	WriteStruct(Schema.Type.AsStruct(), static_cast<FStructSchemaId>(Schema.InnerSchema.Get()), *Value.Struct); break;
+		case EMemberKind::Struct:	WriteStruct(Schema.Type.AsStruct(), Schema.InnerSchema.Get().AsStruct(), *Value.Struct); break;
 		}
 	}
 
@@ -728,15 +727,15 @@ private:
 		}
 	}
 		
-	void WriteStruct(FStructType StructType, FStructSchemaId Schema, const FBuiltStruct& Struct)
+	void WriteStruct(FStructType StructType, FStructId Id, const FBuiltStruct& Struct)
 	{
 		Tmp.Reserve(1024);
-		FMemberWriter(Tmp, Schemas, NewIds, Debug).WriteMembers(Schema, Struct);
+		FMemberWriter(Tmp, Schemas, NewIds, Debug).WriteMembers(Id, Struct);
 		WriteSkippableSlice(Bytes, Tmp);
 		Tmp.Reset();
 	}
 
-	void WriteRange(ERangeSizeType NumType, TConstArrayView<FMemberType> Types, FOptionalSchemaId InnermostSchema, const FBuiltRange* Range)
+	void WriteRange(ERangeSizeType NumType, TConstArrayView<FMemberType> Types, FOptionalInnerId InnermostSchema, const FBuiltRange* Range)
 	{
 		check(Types.Num() > 0);
 		check((Types.Num() > 1) == (Types[0].GetKind() == EMemberKind::Range));
@@ -760,7 +759,7 @@ private:
 			{
 			case EMemberKind::Leaf:		WriteLeaves(Types[0].AsLeaf(), *Range);	break;
 			case EMemberKind::Range:	WriteRanges(Types[0].AsRange().MaxSize, Types.RightChop(1), InnermostSchema, Range->AsRanges()); break;
-			case EMemberKind::Struct:	WriteStructs(Types[0].AsStruct(), static_cast<FStructSchemaId>(InnermostSchema.Get()), Range->AsStructs()); break;
+			case EMemberKind::Struct:	WriteStructs(Types[0].AsStruct(), InnermostSchema.Get().AsStruct(), Range->AsStructs()); break;
 			}
 		}
 	}
@@ -798,13 +797,14 @@ private:
 		Tmp.Reset();
 	}
 
-	void WriteStructs(FStructType StructType, FStructSchemaId Schema, TConstArrayView64<const FBuiltStruct*> Structs)
+	void WriteStructs(FStructType StructType, FStructId BuiltId, TConstArrayView64<const FBuiltStruct*> Structs)
 	{
-		Bytes.Append(reinterpret_cast<uint8*>(&Schema.Idx), StructType.IsDynamic * sizeof(Schema.Idx));
-		WriteSkippableItems(Structs, [=](FMemberWriter& Out, const FBuiltStruct* Struct) { Out.WriteStruct(StructType, Schema, *Struct); });
+		FStructSchemaId WriteId = NewIds.RemapStruct(BuiltId);
+		Bytes.Append(reinterpret_cast<uint8*>(&WriteId.Idx), StructType.IsDynamic * sizeof(WriteId.Idx));
+		WriteSkippableItems(Structs, [=](FMemberWriter& Out, const FBuiltStruct* Struct) { Out.WriteStruct(StructType, BuiltId, *Struct); });
 	}
 	
-	void WriteRanges(ERangeSizeType NumType, TConstArrayView<FMemberType> Types, FOptionalSchemaId InnermostSchema, TConstArrayView64<const FBuiltRange*> Ranges)
+	void WriteRanges(ERangeSizeType NumType, TConstArrayView<FMemberType> Types, FOptionalInnerId InnermostSchema, TConstArrayView64<const FBuiltRange*> Ranges)
 	{	
 		WriteSkippableItems(Ranges, [=](FMemberWriter& Out, const FBuiltRange* Range) { Out.WriteRange(NumType, Types, InnermostSchema, Range); });
 	}
@@ -832,7 +832,7 @@ TConstArrayView<FNameId> FWriter::GetUsedNames() const
 	return NewIds->KeptNames;
 }
 
-FOptionalStructSchemaId FWriter::GetWriteId(FStructSchemaId BuiltId) const
+FOptionalStructSchemaId FWriter::GetWriteId(FStructId BuiltId) const
 {
 	return static_cast<FOptionalStructSchemaId>(NewIds->Structs[BuiltId.Idx]);
 }
@@ -842,7 +842,7 @@ void FWriter::WriteSchemas(TArray64<uint8>& Out) const
 	WriteSchemasImpl(Out, Schemas, *NewIds);
 }
 
-FStructSchemaId FWriter::WriteMembers(TArray64<uint8>& Out, FStructSchemaId BuiltId, const FBuiltStruct& Struct) const
+FStructSchemaId FWriter::WriteMembers(TArray64<uint8>& Out, FStructId BuiltId, const FBuiltStruct& Struct) const
 {
 	return FMemberWriter(Out, Schemas, *NewIds, Debug).WriteMembers(BuiltId, Struct);
 }
diff --git a/Engine/Plugins/Experimental/PlainProps/Source/Public/PlainPropsBind.h b/Engine/Plugins/Experimental/PlainProps/Source/Public/PlainPropsBind.h
--- a/Engine/Plugins/Experimental/PlainProps/Source/Public/PlainPropsBind.h
+++ b/Engine/Plugins/Experimental/PlainProps/Source/Public/PlainPropsBind.h
@@ -133,7 +133,7 @@ static_assert(sizeof(FMemberBindType) == 1);
 // Members are loaded in saved FStructSchema order, not current offset order unless upgrade layer reorders
 struct alignas(/*FRangeBinding*/ 8) FSchemaBinding
 {
-	FStructSchemaId			DeclId;
+	FStructId				DeclId;
 	uint16					NumMembers;
 	uint16					NumInnerSchemas;
 	uint16					NumInnerRanges;
@@ -141,7 +141,7 @@ struct alignas(/*FRangeBinding*/ 8) FSchemaBinding
 
 	const FMemberBindType*	GetInnerRangeTypes() const	{ return Members + NumMembers; }
 	const uint32*			GetOffsets() const			{ return AlignPtr<uint32>(GetInnerRangeTypes() + NumInnerRanges); }
-	const FSchemaId*		GetInnerSchemas() const		{ return AlignPtr<FSchemaId>(GetOffsets() + NumMembers); }
+	const FInnerId*			GetInnerSchemas() const		{ return AlignPtr<FInnerId>(GetOffsets() + NumMembers); }
 	const FRangeBinding*	GetRangeBindings() const	{ return AlignPtr<FRangeBinding>(GetInnerSchemas() + NumInnerSchemas); }
 	uint32					CalculateSize() const;
 	bool					HasSuper() const			{ return NumInnerSchemas > 0 && Members[0].IsStruct() && Members[0].AsStruct().IsSuper; }
@@ -197,7 +197,7 @@ inline FUnpackedLeafType UnpackNonBitfield(FLeafBindType Packed)
 struct FLeafMemberBinding
 {
 	FUnpackedLeafBindType	Leaf;
-	FOptionalEnumSchemaId	Enum;
+	FOptionalEnumId			Enum;
 	SIZE_T					Offset;
 };
 
@@ -206,14 +206,14 @@ struct FRangeMemberBinding
 	const FMemberBindType*	InnerTypes;
 	const FRangeBinding*	RangeBindings;
 	uint16					NumRanges; // At least 1, >1 for nested ranges
-	FOptionalSchemaId		InnermostSchema;
+	FOptionalInnerId		InnermostSchema;
 	SIZE_T					Offset;
 };
 
 struct FStructMemberBinding
 {
 	FStructType				Type;
-	FStructSchemaId			Id;
+	FStructId				Id;
 	SIZE_T					Offset;
 };
 
@@ -221,12 +221,12 @@ PLAINPROPS_API FRangeMemberBinding GetInnerRange(FRangeMemberBinding In); // @pr
 
 ////////////////////////////////////////////////////////////////////////////////////////////////
 
-struct FDualStructSchemaId
+struct FDualStructId
 {
-	FStructSchemaId BindId;
-	FStructSchemaId DeclId;
+	FStructId				BindId;
+	FStructId				DeclId;
 
-	bool IsLowered() const { return BindId != DeclId; }
+	bool					IsLowered() const { return BindId != DeclId; }
 };
 
 ////////////////////////////////////////////////////////////////////////////////////////////////
@@ -234,7 +234,7 @@ struct FDualStructSchemaId
 struct FInnerStruct
 {
 	FMemberId				Name; // Of outer Range or Struct
-	FStructSchemaId			BindId;
+	FStructId				BindId;
 };
 
 enum ECustomLoadMethod { Construct, Assign };
@@ -249,9 +249,9 @@ enum ECustomLoadMethod { Construct, Assign };
 struct ICustomBinding
 {
 	virtual ~ICustomBinding() {}
-	virtual void							SaveCustom(FMemberBuilder& Dst, const void* Src, const void* Default, const FSaveContext& Ctx) = 0;
-	virtual void							LoadCustom(void* Dst, FStructLoadView Src, ECustomLoadMethod Method) const = 0;
-	virtual bool							DiffCustom(const void* StructA, const void* StructB) const = 0;
+	virtual void	SaveCustom(FMemberBuilder& Dst, const void* Src, const void* Default, const FSaveContext& Ctx) = 0;
+	virtual void	LoadCustom(void* Dst, FStructLoadView Src, ECustomLoadMethod Method) const = 0;
+	virtual bool	DiffCustom(const void* StructA, const void* StructB) const = 0;
 };
 
 struct FInnersHandle
@@ -266,7 +266,7 @@ struct FIdWindow
 	uint32			Num = 0;
 };
 
-// Slightly optimized map from FStructSchemaId to ICustomBinding/FStructDeclaration/FInnersHandle
+// Slightly optimized map from FStructId to ICustomBinding/FStructDeclaration/FInnersHandle
 class FCustomBindingMap
 {
 public:
@@ -274,13 +274,13 @@ public:
 	explicit FCustomBindingMap(FDebugIds Dbg) : Debug(Dbg) {}
 	~FCustomBindingMap() { FMemory::Free(Values); }
 
-	void							Bind(FStructSchemaId BindId, ICustomBinding& Binding, const FStructDeclaration& Declaration, FInnersHandle LoweredInners);
-	FCustomBindingHandle			Find(FStructSchemaId BindId) const;
-	void							Drop(FStructSchemaId BindId);
+	void							Bind(FStructId BindId, ICustomBinding& Binding, const FStructDeclaration& Declaration, FInnersHandle LoweredInners);
+	FCustomBindingHandle			Find(FStructId BindId) const;
+	void							Drop(FStructId BindId);
 
 private:
 	FIdWindow						Window;
-	TSet<FStructSchemaId>			Keys;
+	TSet<FStructId>			Keys;
 	ICustomBinding**				Values = nullptr;
 	uint32							MaxValues = 0;
 public:
@@ -291,14 +291,14 @@ class FCustomBindings
 {
 public:
 	// @param Binding must outlive this or call DropStruct()
-	PLAINPROPS_API void							BindStruct(FStructSchemaId BindId, ICustomBinding& Binding, const FStructDeclaration& Declaration, TConstArrayView<FInnerStruct> LoweredInners = {});
-	PLAINPROPS_API const ICustomBinding*		FindStruct(FStructSchemaId BindId) const;
-	PLAINPROPS_API const ICustomBinding*		FindStruct(FStructSchemaId BindId, TConstArrayView<FInnerStruct>& OutLoweredInners) const;
-	PLAINPROPS_API ICustomBinding*				FindStructToSave(FStructSchemaId BindId, const FStructDeclaration*& OutDeclaration) const;
-	PLAINPROPS_API const FStructDeclaration*	FindDeclaration(FStructSchemaId BindId) const;
-	PLAINPROPS_API void							DropStruct(FStructSchemaId BindId);
-
-	virtual FCustomBindingHandle				Find(FStructSchemaId BindId) const = 0;
+	PLAINPROPS_API void							BindStruct(FStructId BindId, ICustomBinding& Binding, const FStructDeclaration& Declaration, TConstArrayView<FInnerStruct> LoweredInners = {});
+	PLAINPROPS_API const ICustomBinding*		FindStruct(FStructId BindId) const;
+	PLAINPROPS_API const ICustomBinding*		FindStruct(FStructId BindId, TConstArrayView<FInnerStruct>& OutLoweredInners) const;
+	PLAINPROPS_API ICustomBinding*				FindStructToSave(FStructId BindId, const FStructDeclaration*& OutDeclaration) const;
+	PLAINPROPS_API const FStructDeclaration*	FindDeclaration(FStructId BindId) const;
+	PLAINPROPS_API void							DropStruct(FStructId BindId);
+
+	virtual FCustomBindingHandle				Find(FStructId BindId) const = 0;
 protected:
 	FCustomBindings(TArray<FInnerStruct>& Inners, FDebugIds Dbg) : Map(Dbg), BottomInners(Inners) {}
 	FCustomBindings(const FCustomBindings* Under) : Map(Under->Map.Debug), BottomInners(Under->BottomInners) {}
@@ -311,7 +311,7 @@ protected:
 class FCustomBindingsBottom final : public FCustomBindings
 {
 	mutable TArray<FInnerStruct>				AllInners;
-	PLAINPROPS_API virtual FCustomBindingHandle	Find(FStructSchemaId BindId) const override;	
+	PLAINPROPS_API virtual FCustomBindingHandle	Find(FStructId BindId) const override;	
 public:
 	explicit FCustomBindingsBottom(FDebugIds Dbg) : FCustomBindings(/* unconstructed */  AllInners, Dbg) {}
 };
@@ -319,7 +319,7 @@ public:
 class FCustomBindingsOverlay final : public FCustomBindings
 {
 	const FCustomBindings&						Underlay;
-	PLAINPROPS_API virtual FCustomBindingHandle	Find(FStructSchemaId BindId) const override;
+	PLAINPROPS_API virtual FCustomBindingHandle	Find(FStructId BindId) const override;
 public:
 	explicit FCustomBindingsOverlay(const FCustomBindings& Under) : FCustomBindings(&Under), Underlay(Under) {}
 };
@@ -594,7 +594,7 @@ struct FMemberBinding
 
 	uint64							Offset;
 	FMemberBindType					InnermostType;		// Always Leaf or Struct
-	FOptionalSchemaId				InnermostSchema;	// Enum or struct schema
+	FOptionalInnerId				InnermostSchema;	// Enum or struct schema
 	TConstArrayView<FRangeBinding>	RangeBindings;		// Non-empty -> Range
 
 	// @pre InnermostSchema isn't type-erased / lowered
@@ -608,12 +608,12 @@ public:
 	explicit FSchemaBindings(FDebugIds In) : Debug(In) {}
 	PLAINPROPS_API ~FSchemaBindings();
 
-	PLAINPROPS_API void						BindStruct(FStructSchemaId BindId, FStructSchemaId DeclId, TConstArrayView<FMemberBinding> Schema);
-	PLAINPROPS_API const FSchemaBinding*	FindStruct(FStructSchemaId BindId) const;
-	PLAINPROPS_API const FSchemaBinding&	GetStruct(FStructSchemaId BindId) const;
-	PLAINPROPS_API void						DropStruct(FStructSchemaId BindId);
+	PLAINPROPS_API void						BindStruct(FStructId BindId, FStructId DeclId, TConstArrayView<FMemberBinding> Schema);
+	PLAINPROPS_API const FSchemaBinding*	FindStruct(FStructId BindId) const;
+	PLAINPROPS_API const FSchemaBinding&	GetStruct(FStructId BindId) const;
+	PLAINPROPS_API void						DropStruct(FStructId BindId);
 
-	PLAINPROPS_API virtual FStructSchemaId	GetDeclId(FStructSchemaId BindId) const override final;
+	PLAINPROPS_API virtual FStructId	GetDeclId(FStructId BindId) const override final;
 private:
 	TArray<TUniquePtr<FSchemaBinding>>		Bindings;
 	FDebugIds								Debug;
@@ -624,7 +624,7 @@ private:
 // Lookup default instances when delta-saving struct ranges
 struct IDefaultStructs
 {
-	virtual const void* Get(FStructSchemaId BindId) = 0;
+	virtual const void* Get(FStructId BindId) = 0;
 };
 
 ////////////////////////////////////////////////////////////////////////////////////////////////
@@ -639,7 +639,7 @@ struct FStructBindIds final : IStructBindIds
 	const FCustomBindings& Customs;
 	const FSchemaBindings& Schemas;
 
-	PLAINPROPS_API virtual FStructSchemaId GetDeclId(FStructSchemaId BindId) const override;
+	PLAINPROPS_API virtual FStructId GetDeclId(FStructId BindId) const override;
 };
 
 ////////////////////////////////////////////////////////////////////////////////////////////////
@@ -688,7 +688,7 @@ FType IndexStructName()
 }
 
 template<class Ids, class Typename>
-FStructSchemaId IndexStructBindIdIfNeeded(FStructSchemaId DeclId)
+FStructId IndexStructBindIdIfNeeded(FStructId DeclId)
 {
 	if constexpr (ExplicitBindName<Typename> || ParametricName<Typename>)
 	{
@@ -702,9 +702,9 @@ FStructSchemaId IndexStructBindIdIfNeeded(FStructSchemaId DeclId)
 }
 
 template<class Ids, class Typename>
-FDualStructSchemaId IndexStructDualId(FType DeclName = IndexStructName<Ids, ETypename::Decl, Typename>())
+FDualStructId IndexStructDualId(FType DeclName = IndexStructName<Ids, ETypename::Decl, Typename>())
 {	
-	FDualStructSchemaId Out;
+	FDualStructId Out;
 	Out.DeclId = Ids::IndexStruct(DeclName);
 	Out.BindId = Out.DeclId;
 
@@ -719,25 +719,25 @@ FDualStructSchemaId IndexStructDualId(FType DeclName = IndexStructName<Ids, ETyp
 
 // Cached by function static
 template<class Ids, typename Struct>
-FStructSchemaId GetStructDeclId()
+FStructId GetStructDeclId()
 {
-	static FStructSchemaId Id = Ids::IndexStruct(IndexStructName<Ids, ETypename::Decl, TTypename<Struct>>());
+	static FStructId Id = Ids::IndexStruct(IndexStructName<Ids, ETypename::Decl, TTypename<Struct>>());
 	return Id;
 }
 
 // Cached by function static
 template<class Ids, typename Struct>
-FStructSchemaId GetStructBindId()
+FStructId GetStructBindId()
 {
-	static FStructSchemaId Id = Ids::IndexStruct(IndexStructName<Ids, ETypename::Bind, TTypename<Struct>>());
+	static FStructId Id = Ids::IndexStruct(IndexStructName<Ids, ETypename::Bind, TTypename<Struct>>());
 	return Id;
 }
 
 // Cached by function static
 template<class Ids, typename Struct>
-FDualStructSchemaId GetStructDualId()
+FDualStructId GetStructDualId()
 {
-	static FDualStructSchemaId Id = IndexStructDualId<Ids, TTypename<Struct>>();
+	static FDualStructId Id = IndexStructDualId<Ids, TTypename<Struct>>();
 	return Id;
 }
 
@@ -755,9 +755,9 @@ FType IndexCttiName()
 
 // Cached by function static
 template<class Ids, typename Enum>
-FEnumSchemaId GetEnumId()
+FEnumId GetEnumId()
 {
-	static FEnumSchemaId Id = Ids::IndexEnum(IndexCttiName<Ids, CttiOf<Enum>>());
+	static FEnumId Id = Ids::IndexEnum(IndexCttiName<Ids, CttiOf<Enum>>());
 	return Id;
 }
 
@@ -829,7 +829,7 @@ struct FCustomInit
 	InnerStructArray& OutLowered;
 
 	// Only needed for inner structs that might get type-erased/lowered
-	void RegisterInnerStruct(FDualStructSchemaId Id, TConstArrayView<FMemberId> Names)
+	void RegisterInnerStruct(FDualStructId Id, TConstArrayView<FMemberId> Names)
 	{
 		if (Id.IsLowered())
 		{
@@ -852,7 +852,7 @@ struct TDependent { using StructLoadView = FStructLoadView; };
 
 // Scoped custom binding for native type
 template<class T, class Runtime, typename CustomBinding = CustomBind<T>>
-struct TScopedStructBinding : FDualStructSchemaId, CustomBinding
+struct TScopedStructBinding : FDualStructId, CustomBinding
 {
 	using Ids = typename Runtime::Ids;
 	using Typename = TCustomTypename<CustomBinding>::Type;
@@ -861,16 +861,16 @@ struct TScopedStructBinding : FDualStructSchemaId, CustomBinding
 	: TScopedStructBinding(DeclName, IndexStructDualId<Ids, Typename>(DeclName), Lowered)
 	{}
 
-	explicit TScopedStructBinding(FDualStructSchemaId Id, InnerStructArray Lowered = {})
+	explicit TScopedStructBinding(FDualStructId Id, InnerStructArray Lowered = {})
 	: TScopedStructBinding(Ids::GetIndexer().Resolve(Id.DeclId), Id, Lowered)
 	{}
 
-	explicit TScopedStructBinding(FStructSchemaId Id, InnerStructArray Lowered = {})
+	explicit TScopedStructBinding(FStructId Id, InnerStructArray Lowered = {})
 	: TScopedStructBinding(Ids::GetIndexer().Resolve(Id), {Id, Id}, Lowered)
 	{}
 
-	TScopedStructBinding(FType DeclName, FDualStructSchemaId DualId, InnerStructArray& Lowered)
-	: FDualStructSchemaId(DualId)
+	TScopedStructBinding(FType DeclName, FDualStructId DualId, InnerStructArray& Lowered)
+	: FDualStructId(DualId)
 	, CustomBinding(TCustomInit<Ids>{{Lowered}})
 	{
 		const FStructDeclaration& Declaration = Runtime::GetTypes().DeclareStruct(DeclId, DeclName, CustomBinding::MemberIds, TOccupancyOf<T>::Value);
@@ -906,8 +906,8 @@ struct TScopedStructBinding<T, Runtime, void>
 	using Ids = typename Runtime::Ids;
 	using Typename = TTypename<T>;
 	
-	const FStructSchemaId DeclId;
-	const FStructSchemaId BindId;
+	const FStructId DeclId;
+	const FStructId BindId;
 
 	TScopedStructBinding(FType DeclName = IndexStructName<Ids, ETypename::Decl, Typename>())
 	: DeclId(DeclareNativeStruct<CttiOf<T>, Ids>(Runtime::GetTypes(), TOccupancyOf<T>::Value))
@@ -924,7 +924,7 @@ struct TScopedStructBinding<T, Runtime, void>
 };
 
 template<class Struct, typename CustomBinding, class Runtime>
-FStructSchemaId BindCustomStructOnce()
+FStructId BindCustomStructOnce()
 {
 	static TScopedStructBinding<Struct, Runtime, CustomBinding> Instance;
 	return Instance.BindId;
@@ -934,39 +934,39 @@ inline constexpr FMemberBindType DefaultStructBindType = FMemberBindType(FStruct
 inline constexpr FMemberBindType SuperStructBindType = FMemberBindType(FStructType{EMemberKind::Struct, /* IsDynamic */ 0, /* IsSuper */ 1});
 
 template<class Struct, class CustomBinding, class Runtime>
-FMemberBindType BindMemberStruct(FOptionalSchemaId& OutSchema)
+FMemberBindType BindMemberStruct(FOptionalInnerId& OutSchema)
 {
 	// Todo: Consider auto-binding all nested structs
 	if constexpr (std::is_void_v<CustomBinding>)
 	{
-		OutSchema = FOptionalSchemaId(GetStructBindId<typename Runtime::Ids, Struct>());
+		OutSchema = FInnerId(GetStructBindId<typename Runtime::Ids, Struct>());
 	}
 	else
 	{
-		OutSchema = FOptionalSchemaId(BindCustomStructOnce<Struct, CustomBinding, Runtime>());
+		OutSchema = FInnerId(BindCustomStructOnce<Struct, CustomBinding, Runtime>());
 	}
 
 	return DefaultStructBindType;
 }
 
 template<typename Struct, class Ids>
-FMemberBindType BindInnermostType(FOptionalSchemaId& OutSchema)
+FMemberBindType BindInnermostType(FOptionalInnerId& OutSchema)
 {
-	OutSchema = FOptionalSchemaId(GetStructBindId<Ids, Struct>());
+	OutSchema = FInnerId(GetStructBindId<Ids, Struct>());
 	return DefaultStructBindType;
 }
 
 template<Arithmetic Type, class Ids>
-FMemberBindType BindInnermostType(FOptionalSchemaId& OutSchema)
+FMemberBindType BindInnermostType(FOptionalInnerId& OutSchema)
 {
 	OutSchema = NoId;
 	return FMemberBindType(ReflectArithmetic<Type>);
 }
 
 template<Enumeration Enum, class Ids>
-FMemberBindType BindInnermostType(FOptionalSchemaId& OutSchema)
+FMemberBindType BindInnermostType(FOptionalInnerId& OutSchema)
 {
-	OutSchema = FOptionalSchemaId(GetEnumId<Ids, Enum>());
+	OutSchema = FInnerId(GetEnumId<Ids, Enum>());
 	return FMemberBindType(ReflectEnum<Enum>);
 }
 
@@ -1079,12 +1079,12 @@ FMemberBinding BindMember(uint64 Offset)
 ////////////////////////////////////////////////////////////////////////////////////////////////
 
 template<class Ctti, class Ids>
-FEnumSchemaId DeclareNativeEnum(FDeclarations& Out, EEnumMode Mode)
+FEnumId DeclareNativeEnum(FDeclarations& Out, EEnumMode Mode)
 {
 	using UnderlyingType = std::underlying_type_t<typename Ctti::Type>;
 
 	FType Type = IndexCttiName<Ids, Ctti>();
-	FEnumSchemaId Id = Ids::IndexEnum(Type);
+	FEnumId Id = Ids::IndexEnum(Type);
 	FEnumerator Enumerators[Ctti::NumEnumerators];
 	for (FEnumerator& Enumerator : Enumerators)
 	{
@@ -1097,14 +1097,14 @@ FEnumSchemaId DeclareNativeEnum(FDeclarations& Out, EEnumMode Mode)
 }
 
 template<class Ctti, class Ids>
-FStructSchemaId DeclareNativeStruct(FDeclarations& Out, EMemberPresence Occupancy)
+FStructId DeclareNativeStruct(FDeclarations& Out, EMemberPresence Occupancy)
 {
 	using Typename = TTypename<typename Ctti::Type>;
 	using SuperType = typename Ctti::Super;
 
 	FType Type = IndexStructName<Ids, ETypename::Decl, Typename>();
-	FStructSchemaId Id = Ids::IndexStruct(Type);
-	FOptionalStructSchemaId SuperId;
+	FStructId Id = Ids::IndexStruct(Type);
+	FOptionalStructId SuperId;
 	if constexpr (!std::is_void_v<SuperType>)
 	{
 		SuperId = GetStructDeclId<Ids, SuperType>();
@@ -1121,7 +1121,7 @@ FStructSchemaId DeclareNativeStruct(FDeclarations& Out, EMemberPresence Occupanc
 }
 
 template<class Ctti, class Runtime>
-void BindNativeStruct(FSchemaBindings& Out, FStructSchemaId BindId, FStructSchemaId DeclId)
+void BindNativeStruct(FSchemaBindings& Out, FStructId BindId, FStructId DeclId)
 {
 	FMemberBinding MemberBindings[Ctti::NumVars];
 	ForEachVar<Ctti>([&]<class Var>()
@@ -1134,7 +1134,7 @@ void BindNativeStruct(FSchemaBindings& Out, FStructSchemaId BindId, FStructSchem
 //////////////////////////////////////////////////////////////////////////
 
 // Save -> load struct ids for ESchemaFormat::InMemoryNames, alternative to side-channel with ExtractRuntimeIds()
-[[nodiscard]] PLAINPROPS_API TArray<FStructSchemaId> IndexRuntimeIds(const FSchemaBatch& Schemas,  FIdIndexerBase& Indexer);
+[[nodiscard]] PLAINPROPS_API TArray<FStructId> IndexRuntimeIds(const FSchemaBatch& Schemas,  FIdIndexerBase& Indexer);
 
 // Save -> load ids for ESchemaFormat::StableNames
 struct FIdBinding
@@ -1142,7 +1142,7 @@ struct FIdBinding
 	TConstArrayView<FNameId>			Names;
 	TConstArrayView<FNestedScopeId>		NestedScopes;
 	TConstArrayView<FParametricTypeId>	ParametricTypes;
-	TConstArrayView<FSchemaId>			Schemas;
+	TConstArrayView<FInnerId>			Schemas;
 
 	FNameId								Remap(FNameId Old) const				{ return Names[Old.Idx]; }
 	FMemberId							Remap(FMemberId Old) const				{ return { Remap(Old.Id) }; }
@@ -1157,11 +1157,11 @@ struct FIdBinding
 	template<typename T>
 	TOptionalId<T>						Remap(TOptionalId<T> Old) const			{ return Old ? ToOptional(Remap(Old.Get())) : Old; }
 
-	TConstArrayView<FStructSchemaId>	GetStructIds(int32 NumStructs) const
+	TConstArrayView<FStructId>			GetStructIds(int32 NumStructs) const
 	{
 		// All saved struct schema ids are lower than enum schema ids
 		check(NumStructs <= Schemas.Num());
-		return MakeArrayView(static_cast<const FStructSchemaId*>(Schemas.GetData()), NumStructs);
+		return MakeArrayView(reinterpret_cast<const FStructId*>(Schemas.GetData()), NumStructs);
 	}
 };
 
diff --git a/Engine/Plugins/Experimental/PlainProps/Source/Public/PlainPropsBuild.h b/Engine/Plugins/Experimental/PlainProps/Source/Public/PlainPropsBuild.h
--- a/Engine/Plugins/Experimental/PlainProps/Source/Public/PlainPropsBuild.h
+++ b/Engine/Plugins/Experimental/PlainProps/Source/Public/PlainPropsBuild.h
@@ -81,7 +81,7 @@ struct FMemberSchema
 	FMemberType						Type;
 	FMemberType						InnerRangeType;
 	uint16							NumInnerRanges;
-	FOptionalSchemaId				InnerSchema;
+	FOptionalInnerId				InnerSchema;
 	const FMemberType*				NestedRangeTypes;
 
 	TConstArrayView<FMemberType> GetInnerRangeTypes() const
@@ -150,21 +150,21 @@ FMemberSchema MakeLeafRangeSchema()
 }
 
 template<Enumeration T>
-FMemberSchema MakeEnumRangeSchema(FEnumSchemaId Schema, ERangeSizeType MaxSize)
+FMemberSchema MakeEnumRangeSchema(FEnumId Schema, ERangeSizeType MaxSize)
 {
-	return { FMemberType(MaxSize), ReflectEnum<T>.Pack(), 1, Schema, nullptr };
+	return { FMemberType(MaxSize), ReflectEnum<T>.Pack(), 1, FInnerId(Schema), nullptr };
 }
 
 template<Enumeration T, typename SizeType>
-FMemberSchema MakeEnumRangeSchema(FEnumSchemaId Schema)
+FMemberSchema MakeEnumRangeSchema(FEnumId Schema)
 {
 	return MakeEnumRangeSchema<T>(Schema, RangeSizeOf(SizeType{}));
 }
 
 template<UnsignedIntegral T>
-FMemberSchema MakeEnumRangeSchema(FEnumSchemaId Schema, ERangeSizeType MaxSize)
+FMemberSchema MakeEnumRangeSchema(FEnumId Schema, ERangeSizeType MaxSize)
 {
-	return { FMemberType(MaxSize), FMemberType(ELeafType::Enum, ReflectArithmetic<T>.Width), 1, Schema, nullptr };
+	return { FMemberType(MaxSize), FMemberType(ELeafType::Enum, ReflectArithmetic<T>.Width), 1, FInnerId(Schema), nullptr };
 }
 
 inline constexpr FMemberType DefaultStructType =	FMemberType(FStructType{EMemberKind::Struct, /* IsDynamic */ 0, /* IsSuper */ 0});
@@ -172,9 +172,9 @@ inline constexpr FMemberType SuperStructType =		FMemberType(FStructType{EMemberK
 inline constexpr FMemberType DynamicStructType =	FMemberType(FStructType{EMemberKind::Struct, /* IsDynamic */ 1, /* IsSuper */ 0});
 
 
-inline FMemberSchema MakeStructRangeSchema(ERangeSizeType SizeType, FStructSchemaId Schema)
+inline FMemberSchema MakeStructRangeSchema(ERangeSizeType SizeType, FStructId Schema)
 {
-	return { FMemberType(SizeType), DefaultStructType, 1, Schema, nullptr };
+	return { FMemberType(SizeType), DefaultStructType, 1, FInnerId(Schema), nullptr };
 }
 
 inline FMemberSchema MakeDynamicStructRangeSchema(ERangeSizeType SizeType)
@@ -214,18 +214,18 @@ template<Arithmetic T, typename SizeType>
 }
 
 template<Enumeration T, typename SizeType>
-[[nodiscard]] FTypedRange BuildEnumRange(FScratchAllocator& Scratch, FEnumSchemaId Enum, TConstArrayView<T, SizeType> Values)
+[[nodiscard]] FTypedRange BuildEnumRange(FScratchAllocator& Scratch, FEnumId Enum, TConstArrayView<T, SizeType> Values)
 {
 	return { MakeEnumRangeSchema<T, SizeType>(Enum), CloneLeaves(Scratch, Values.Num(), Values.GetData(), sizeof(T)) };
 }
 
 template<UnsignedIntegral T>
-[[nodiscard]] FTypedRange BuildEnumRange(FScratchAllocator& Scratch, FEnumSchemaId Enum, ERangeSizeType SizeType, TConstArrayView<T> Values)
+[[nodiscard]] FTypedRange BuildEnumRange(FScratchAllocator& Scratch, FEnumId Enum, ERangeSizeType SizeType, TConstArrayView<T> Values)
 {
 	return { MakeEnumRangeSchema<T>(Enum, SizeType), CloneLeaves(Scratch, Values.Num(), Values.GetData(), sizeof(T)) };
 }
 
-[[nodiscard]] inline FTypedRange MakeStructRange(FStructSchemaId Schema, ERangeSizeType SizeType, FBuiltRange* Values )
+[[nodiscard]] inline FTypedRange MakeStructRange(FStructId Schema, ERangeSizeType SizeType, FBuiltRange* Values )
 {
 	return { MakeStructRangeSchema(SizeType, Schema), Values };
 }
@@ -249,10 +249,10 @@ struct FBuiltMember
 {
 	FBuiltMember(FMemberId InName, FTypedValue In) : FBuiltMember(InName, In.Schema, In.Value) {}
 	FBuiltMember(FOptionalMemberId N, FMemberSchema S, FBuiltValue V) : Name(N), Schema(MoveTemp(S)), Value(V) {}
-	PLAINPROPS_API FBuiltMember(FMemberId Name, FUnpackedLeafType Leaf, FOptionalEnumSchemaId Schema, uint64 Value);
+	PLAINPROPS_API FBuiltMember(FMemberId Name, FUnpackedLeafType Leaf, FOptionalEnumId Schema, uint64 Value);
 	PLAINPROPS_API FBuiltMember(FMemberId Name, FTypedRange Range);
-	PLAINPROPS_API FBuiltMember(FMemberId Name, FStructSchemaId Schema, FBuiltStruct* Value);
-	PLAINPROPS_API static FBuiltMember MakeSuper(FStructSchemaId Schema, FBuiltStruct* Value);
+	PLAINPROPS_API FBuiltMember(FMemberId Name, FStructId Schema, FBuiltStruct* Value);
+	PLAINPROPS_API static FBuiltMember MakeSuper(FStructId Schema, FBuiltStruct* Value);
 
 	FOptionalMemberId		Name;
 	FMemberSchema			Schema;
@@ -272,13 +272,13 @@ public:
 	}
 
 	template<Enumeration T>
-	void AddEnum(FMemberId Name, FEnumSchemaId Schema, T Value)
+	void AddEnum(FMemberId Name, FEnumId Schema, T Value)
 	{
 		AddLeaf(Name, ReflectEnum<T>, ToOptional(Schema), ValueCast(static_cast<__underlying_type(T)>(Value)));
 	}
 	
 	template<UnsignedIntegral T>
-	void AddEnum(FMemberId Name, FEnumSchemaId Schema, T Value)
+	void AddEnum(FMemberId Name, FEnumId Schema, T Value)
 	{
 		AddLeaf(Name, {ELeafType::Enum, ReflectArithmetic<T>.Width}, ToOptional(Schema), Value);
 	}
@@ -289,14 +289,14 @@ public:
 		AddLeaf(Name, {ELeafType::Hex, ReflectArithmetic<T>.Width}, NoId, Value);
 	}
 
-	void AddLeaf(FMemberId Name, FUnpackedLeafType Leaf, FOptionalEnumSchemaId Enum, uint64 Value)	{ Members.Emplace(Name, Leaf, Enum, Value); }
-	void AddRange(FMemberId Name, FTypedRange Range)												{ Members.Emplace(Name, Range); }
+	void AddLeaf(FMemberId Name, FUnpackedLeafType Leaf, FOptionalEnumId Enum, uint64 Value)	{ Members.Emplace(Name, Leaf, Enum, Value); }
+	void AddRange(FMemberId Name, FTypedRange Range)											{ Members.Emplace(Name, Range); }
 	// Add already built nested struct, must not be null
-	void AddStruct(FMemberId Name, FStructSchemaId Schema, FBuiltStruct* Struct)					{ check(Struct); Members.Emplace(Name, Schema, Struct); }
-	void Add(FMemberId Name, FTypedValue TypedValue)												{ Members.Emplace(Name, TypedValue); }
+	void AddStruct(FMemberId Name, FStructId Schema, FBuiltStruct* Struct)						{ check(Struct); Members.Emplace(Name, Schema, Struct); }
+	void Add(FMemberId Name, FTypedValue TypedValue)											{ Members.Emplace(Name, TypedValue); }
 
 	// Add already built super struct, must not be null and must be the first member added
-	PLAINPROPS_API void AddSuperStruct(FStructSchemaId SuperSchema, FBuiltStruct* SuperStruct);
+	PLAINPROPS_API void AddSuperStruct(FStructId SuperSchema, FBuiltStruct* SuperStruct);
 	
 	// Build members into a single nested super struct member, no-op if no non-super members has been added
 	PLAINPROPS_API void BuildSuperStruct(FScratchAllocator&	Scratch, const FStructDeclaration& Super, const FDebugIds& Debug);
diff --git a/Engine/Plugins/Experimental/PlainProps/Source/Public/PlainPropsBuildSchema.h b/Engine/Plugins/Experimental/PlainProps/Source/Public/PlainPropsBuildSchema.h
--- a/Engine/Plugins/Experimental/PlainProps/Source/Public/PlainPropsBuildSchema.h
+++ b/Engine/Plugins/Experimental/PlainProps/Source/Public/PlainPropsBuildSchema.h
@@ -18,8 +18,8 @@ struct FStructSchemaBuilder;
 struct FBuiltStructSchema
 {
 	FType							Type;
-	FStructSchemaId					Id;
-	FOptionalStructSchemaId			Super;
+	FStructId						Id;
+	FOptionalStructId				Super;
 	bool							bDense = false;
 	TArray<FMemberId>				MemberNames;
 	TArray<const FMemberSchema*>	MemberSchemas;
@@ -28,7 +28,7 @@ struct FBuiltStructSchema
 struct FBuiltEnumSchema
 {
 	FType							Type;
-	FEnumSchemaId					Id;
+	FEnumId							Id;
 	EEnumMode						Mode = EEnumMode::Flat;
 	ELeafWidth						Width = ELeafWidth::B8;
 	TArray<FNameId>					Names;
@@ -51,9 +51,9 @@ public:
 	PLAINPROPS_API FSchemasBuilder(FStructDeclarations InStructs, FEnumDeclarations InEnums, const FIds& InIds, const IStructBindIds& BindIds, FScratchAllocator& Scratch, ESchemaFormat InFormat);
 	PLAINPROPS_API ~FSchemasBuilder();
 
-	PLAINPROPS_API FEnumSchemaBuilder&			NoteEnum(FEnumSchemaId Id);
-	PLAINPROPS_API FStructSchemaBuilder&		NoteStruct(FStructSchemaId BindId);
-	PLAINPROPS_API void							NoteStructAndMembers(FStructSchemaId BindId, const FBuiltStruct& Struct);
+	PLAINPROPS_API FEnumSchemaBuilder&			NoteEnum(FEnumId Id);
+	PLAINPROPS_API FStructSchemaBuilder&		NoteStruct(FStructId BindId);
+	PLAINPROPS_API void							NoteStructAndMembers(FStructId BindId, const FBuiltStruct& Struct);
 	PLAINPROPS_API FBuiltSchemas				Build();
 	
 	FScratchAllocator&							GetScratch() const { return Scratch; }
@@ -63,8 +63,8 @@ public:
 private:
 	FStructDeclarations							DeclaredStructs;
 	FEnumDeclarations							DeclaredEnums;
-	TSet<FStructSchemaId>						StructIndices;
-	TSet<FEnumSchemaId>							EnumIndices;
+	TSet<FStructId>								StructIndices;
+	TSet<FEnumId>								EnumIndices;
 	const FIds&									Ids;
 	const IStructBindIds&						BindIds;
 	ESchemaFormat								Format;
@@ -78,6 +78,6 @@ private:
 };
 
 // Extract runtime ids to side-channel for loading to avoid reindexing with IndexRuntimeIds()
-[[nodiscard]] PLAINPROPS_API TArray<FStructSchemaId> ExtractRuntimeIds(const FBuiltSchemas& In);
+[[nodiscard]] PLAINPROPS_API TArray<FStructId> ExtractRuntimeIds(const FBuiltSchemas& In);
 
 } // namespace PlainProps
diff --git a/Engine/Plugins/Experimental/PlainProps/Source/Public/PlainPropsDeclare.h b/Engine/Plugins/Experimental/PlainProps/Source/Public/PlainPropsDeclare.h
--- a/Engine/Plugins/Experimental/PlainProps/Source/Public/PlainPropsDeclare.h
+++ b/Engine/Plugins/Experimental/PlainProps/Source/Public/PlainPropsDeclare.h
@@ -37,9 +37,9 @@ enum class EMemberPresence : uint8 { RequireAll, AllowSparse };
 struct FStructDeclaration
 {
 	uint32					RefCount;
-	FStructSchemaId			Id;				// Could be removed, might allow declaration dedup among templated types
+	FStructId				Id;				// Could be removed, might allow declaration dedup among templated types
 	FType					Type;			// Could be removed, might allow declaration dedup among templated types
-	FOptionalStructSchemaId	Super;
+	FOptionalStructId		Super;
 	EMemberPresence			Occupancy;
 	uint16					NumMembers;
 	FMemberId				MemberOrder[0];
@@ -62,17 +62,17 @@ public:
 	UE_NONCOPYABLE(FDeclarations);
 	explicit FDeclarations(FDebugIds In) : Debug(In) {}
 
-	PLAINPROPS_API const FEnumDeclaration&			DeclareEnum(FEnumSchemaId DeclId, FType Type, EEnumMode Mode, TConstArrayView<FEnumerator> Enumerators, EEnumAliases Policy);
+	PLAINPROPS_API const FEnumDeclaration&			DeclareEnum(FEnumId DeclId, FType Type, EEnumMode Mode, TConstArrayView<FEnumerator> Enumerators, EEnumAliases Policy);
 	// Declare struct with ref count 1 or increment it and check that previous declaration matches
-	PLAINPROPS_API const FStructDeclaration&		DeclareStruct(FStructSchemaId DeclId, FType Type, TConstArrayView<FMemberId> MemberOrder, EMemberPresence Occupancy, FOptionalStructSchemaId Super = {});
-	PLAINPROPS_API const FStructDeclaration&		DeclareNumeralStruct(FStructSchemaId DeclId, FType Type, TConstArrayView<FMemberId> Numerals, EMemberPresence Occupancy);
+	PLAINPROPS_API const FStructDeclaration&		DeclareStruct(FStructId DeclId, FType Type, TConstArrayView<FMemberId> MemberOrder, EMemberPresence Occupancy, FOptionalStructId Super = {});
+	PLAINPROPS_API const FStructDeclaration&		DeclareNumeralStruct(FStructId DeclId, FType Type, TConstArrayView<FMemberId> Numerals, EMemberPresence Occupancy);
 	
-	void											DropEnum(FEnumSchemaId DeclId)			{ Check(DeclId); DeclaredEnums[DeclId.Idx].Reset(); }
-	PLAINPROPS_API void								DropStructRef(FStructSchemaId DeclId);
+	void											DropEnum(FEnumId DeclId)		{ Check(DeclId); DeclaredEnums[DeclId.Idx].Reset(); }
+	PLAINPROPS_API void								DropStructRef(FStructId DeclId);
 
-	const FEnumDeclaration&							Get(FEnumSchemaId DeclId) const			{ Check(DeclId); return *DeclaredEnums[DeclId.Idx]; }
-	const FStructDeclaration&						Get(FStructSchemaId DeclId) const		{ Check(DeclId); return *DeclaredStructs[DeclId.Idx]; }
-	PLAINPROPS_API const FStructDeclaration*		Find(FStructSchemaId Id) const;
+	const FEnumDeclaration&							Get(FEnumId DeclId) const		{ Check(DeclId); return *DeclaredEnums[DeclId.Idx]; }
+	const FStructDeclaration&						Get(FStructId DeclId) const		{ Check(DeclId); return *DeclaredStructs[DeclId.Idx]; }
+	PLAINPROPS_API const FStructDeclaration*		Find(FStructId Id) const;
 	
 	TConstArrayView<TUniquePtr<FEnumDeclaration>>	GetEnums() const				{ return DeclaredEnums; }
 	TConstArrayView<TUniquePtr<FStructDeclaration>>	GetStructs() const				{ return DeclaredStructs; }
@@ -84,8 +84,8 @@ protected:
 	FDebugIds										Debug;
 
 #if DO_CHECK
-	PLAINPROPS_API void								Check(FEnumSchemaId Id) const;
-	PLAINPROPS_API void								Check(FStructSchemaId Id) const;
+	PLAINPROPS_API void								Check(FEnumId Id) const;
+	PLAINPROPS_API void								Check(FStructId Id) const;
 #else
 	void											Check(...) const {}
 #endif
@@ -94,7 +94,7 @@ protected:
 struct IStructBindIds
 {
 	virtual ~IStructBindIds() = default;
-	virtual FStructSchemaId GetDeclId(FStructSchemaId BindId) const = 0;
+	virtual FStructId GetDeclId(FStructId BindId) const = 0;
 };
 
 } // namespace PlainProps
diff --git a/Engine/Plugins/Experimental/PlainProps/Source/Public/PlainPropsDiff.h b/Engine/Plugins/Experimental/PlainProps/Source/Public/PlainPropsDiff.h
--- a/Engine/Plugins/Experimental/PlainProps/Source/Public/PlainPropsDiff.h
+++ b/Engine/Plugins/Experimental/PlainProps/Source/Public/PlainPropsDiff.h
@@ -9,9 +9,9 @@ namespace PlainProps
 
 union FDiffMetadata
 {
-	FOptionalEnumSchemaId	Leaf;
+	FOptionalEnumId			Leaf;
 	FRangeBinding			Range;
-	FStructSchemaId			Struct;
+	FStructId				Struct;
 };
 
 // Currently lacking range indices
@@ -34,6 +34,6 @@ struct FDiffContext
 	FDiffPath				Out;
 };
 
-[[nodiscard]] PLAINPROPS_API bool DiffStructs(FStructSchemaId BindId, const void* A, const void* B, FDiffContext& Ctx);
+[[nodiscard]] PLAINPROPS_API bool DiffStructs(FStructId BindId, const void* A, const void* B, FDiffContext& Ctx);
 
 } // namespace PlainProps
\ No newline at end of file
diff --git a/Engine/Plugins/Experimental/PlainProps/Source/Public/PlainPropsIndex.h b/Engine/Plugins/Experimental/PlainProps/Source/Public/PlainPropsIndex.h
--- a/Engine/Plugins/Experimental/PlainProps/Source/Public/PlainPropsIndex.h
+++ b/Engine/Plugins/Experimental/PlainProps/Source/Public/PlainPropsIndex.h
@@ -57,8 +57,8 @@ public:
 	FType								MakeLeafParameter(FUnpackedLeafType Leaf)			{ return { NoId, FTypenameId(Leaves[uint8(Leaf.Type)][uint8(Leaf.Width)]) }; }
 	FType								MakeRangeParameter(ERangeSizeType SizeType)			{ return { NoId, FTypenameId(Ranges[uint8(SizeType)]) }; }
 
-	PLAINPROPS_API FEnumSchemaId		IndexEnum(FType Type);
-	PLAINPROPS_API FStructSchemaId		IndexStruct(FType Type);
+	PLAINPROPS_API FEnumId				IndexEnum(FType Type);
+	PLAINPROPS_API FStructId			IndexStruct(FType Type);
 	
 	virtual uint32						NumNestedScopes() const override final				{ return IntCastChecked<uint32>(NestedScopes.Num()); }
 	virtual uint32						NumParametricTypes() const override final			{ return IntCastChecked<uint32>(ParametricTypes.Num()); }
@@ -67,8 +67,8 @@ public:
 	
 	virtual FNestedScope				Resolve(FNestedScopeId Id) const override final		{ return NestedScopes.Resolve(Id); }
 	virtual FParametricTypeView			Resolve(FParametricTypeId Id) const override final	{ return ParametricTypes.Resolve(Id); }
-	PLAINPROPS_API virtual FType		Resolve(FEnumSchemaId Id) const override final;
-	PLAINPROPS_API virtual FType		Resolve(FStructSchemaId Id) const override final;
+	PLAINPROPS_API virtual FType		Resolve(FEnumId Id) const override final;
+	PLAINPROPS_API virtual FType		Resolve(FStructId Id) const override final;
 	
 	const FNestedScopeIndexer&			GetNestedScopes() const								{ return NestedScopes; }
 	const FParametricTypeIndexer&		GetParametricTypes() const							{ return ParametricTypes; }
@@ -93,6 +93,8 @@ template<class NameT>
 class TIdIndexer : public FIdIndexerBase
 {
 public:
+	using FIdIndexerBase::AppendString;
+
 	TIdIndexer() { InitParameterNames(); }
 
 	template<typename T> FNameId				MakeName(T&& Name)						{ return { static_cast<uint32>(Names.Add(NameT(Name)).AsInteger()) }; }
@@ -110,7 +112,7 @@ public:
 protected:
 	TSet<NameT>									Names;
 
-	virtual void								AppendString(FUtf8StringBuilderBase& Out, FNameId Id) const override final
+	virtual void								AppendString(FUtf8Builder& Out, FNameId Id) const override final
 	{
 		PlainProps::AppendString(Out, ResolveName(Id));
 	}
diff --git a/Engine/Plugins/Experimental/PlainProps/Source/Public/PlainPropsLoad.h b/Engine/Plugins/Experimental/PlainProps/Source/Public/PlainPropsLoad.h
--- a/Engine/Plugins/Experimental/PlainProps/Source/Public/PlainPropsLoad.h
+++ b/Engine/Plugins/Experimental/PlainProps/Source/Public/PlainPropsLoad.h
@@ -22,8 +22,8 @@ struct FLoadBatchDeleter
 using FLoadBatchPtr = TUniquePtr<FLoadBatch, FLoadBatchDeleter>;
 
 
-[[nodiscard]] PLAINPROPS_API FLoadBatchPtr	CreateLoadPlans(FReadBatchId ReadId, const FDeclarations& Declarations, const FCustomBindings& Customs, const FSchemaBindings& Schemas, TConstArrayView<FStructSchemaId> RuntimeIds, ESchemaFormat Format);
-PLAINPROPS_API void							LoadStruct(void* Dst, FByteReader Src, FStructSchemaId Id, const FLoadBatch& Batch);
+[[nodiscard]] PLAINPROPS_API FLoadBatchPtr	CreateLoadPlans(FReadBatchId ReadId, const FDeclarations& Declarations, const FCustomBindings& Customs, const FSchemaBindings& Schemas, TConstArrayView<FStructId> RuntimeIds, ESchemaFormat Format);
+PLAINPROPS_API void							LoadStruct(void* Dst, FByteReader Src, FStructSchemaId LoadId, const FLoadBatch& Batch);
 PLAINPROPS_API void							ConstructAndLoadStruct(void* Dst, FByteReader Src, FStructSchemaId LoadId, const FLoadBatch& Batch);
 
 } // namespace PlainProps
\ No newline at end of file
diff --git a/Engine/Plugins/Experimental/PlainProps/Source/Public/PlainPropsPrint.h b/Engine/Plugins/Experimental/PlainProps/Source/Public/PlainPropsPrint.h
--- a/Engine/Plugins/Experimental/PlainProps/Source/Public/PlainPropsPrint.h
+++ b/Engine/Plugins/Experimental/PlainProps/Source/Public/PlainPropsPrint.h
@@ -8,11 +8,11 @@
 namespace PlainProps
 {
 
+class FBatchIds;
+struct FDiffPath;
+struct FMemberSchema;
 class FReadBatchId;
 struct FStructView;
-struct FDiffPath;
-
-///////////////////////////////////////////////////////////////////////////////
 
 PLAINPROPS_API FAnsiStringView ToString(ERangeSizeType Width);
 PLAINPROPS_API FAnsiStringView ToString(FUnpackedLeafType Leaf);
@@ -20,19 +20,8 @@ PLAINPROPS_API FAnsiStringView ToString(ELeafWidth Width);
 
 ///////////////////////////////////////////////////////////////////////////////
 
-PLAINPROPS_API void PrintMemberSchema(
-	FUtf8Builder& Out,
-	const FIds& Ids, 
-	FMemberType Type,
-	FOptionalSchemaId InnerSchema,
-	TConstArrayView<FMemberType> InnerRangeTypes);
-
-PLAINPROPS_API void PrintYamlBatch(
-	FUtf8Builder& Out,
-	const FIds& Ids,
-	FReadBatchId Batch,
-	TConstArrayView<FStructView> Objects);
-
-PLAINPROPS_API void PrintDiff(FUtf8Builder& Out, const FDiffPath& Diff, const FIds& Ids);
+PLAINPROPS_API void PrintMemberSchema(FUtf8Builder& Out, const FIds& Ids, FMemberSchema Member);
+PLAINPROPS_API void PrintYamlBatch(FUtf8Builder& Out, const FBatchIds& Ids, TConstArrayView<FStructView> Objects);
+PLAINPROPS_API void PrintDiff(FUtf8Builder& Out, const FIds& Ids, const FDiffPath& Diff);
 
 } // namespace PlainProps
diff --git a/Engine/Plugins/Experimental/PlainProps/Source/Public/PlainPropsRead.h b/Engine/Plugins/Experimental/PlainProps/Source/Public/PlainPropsRead.h
--- a/Engine/Plugins/Experimental/PlainProps/Source/Public/PlainPropsRead.h
+++ b/Engine/Plugins/Experimental/PlainProps/Source/Public/PlainPropsRead.h
@@ -720,4 +720,64 @@ private:
 	}
 };
 
+//////////////////////////////////////////////////////////////////////////
+
+// Serialized id resolver
+class FBatchIds : public FIdsBase
+{
+public:
+	using FIdsBase::Resolve;
+	using FIdsBase::AppendString;
+
+	PLAINPROPS_API FBatchIds(FReadBatchId Batch);
+
+	PLAINPROPS_API virtual uint32	NumEnums() const;
+	PLAINPROPS_API virtual uint32	NumStructs() const;
+	PLAINPROPS_API virtual FType	Resolve(FEnumSchemaId Id) const;
+	PLAINPROPS_API virtual FType	Resolve(FStructSchemaId Id) const;
+
+	PLAINPROPS_API virtual void		AppendString(FUtf8Builder& Out, FEnumSchemaId Enum) const;
+	PLAINPROPS_API virtual void		AppendString(FUtf8Builder& Out, FStructSchemaId Struct) const;
+
+	const FSchemaBatch& GetSchemas() const { return Schemas; } 
+	FReadBatchId GetBatchId() const { return BatchId; } 
+protected:
+	const FSchemaBatch& Schemas;
+	FReadBatchId BatchId;
+};
+
+// Serialized id resolver for ESchemaFormat::InMemoryNames
+class FMemoryBatchIds final : public FBatchIds
+{
+public:
+	using FBatchIds::Resolve;
+	using FBatchIds::AppendString;
+
+	FMemoryBatchIds(FReadBatchId Batch, const FIdsBase& InNames) : FBatchIds(Batch), Names(InNames) {}
+
+	virtual uint32				NumNames() const override						{ return Names.NumNames(); }
+	virtual uint32				NumNestedScopes() const override				{ return Names.NumNestedScopes(); }
+	virtual uint32				NumParametricTypes() const override				{ return Names.NumParametricTypes(); }
+	virtual FNestedScope		Resolve(FNestedScopeId Id) const override		{ return Names.Resolve(Id); }
+	virtual FParametricTypeView	Resolve(FParametricTypeId Id) const override	{ return Names.Resolve(Id); }
+
+	virtual void				AppendString(FUtf8Builder& Out, FNameId Name) const override { Names.AppendString(Out, Name); }
+
+private:
+	const FIdsBase& Names;
+};
+
+// Serialized id resolver for ESchemaFormat::StableNames
+class FStableBatchIds : public FBatchIds
+{
+public:
+	using FBatchIds::FBatchIds;
+	using FBatchIds::Resolve;
+
+	PLAINPROPS_API virtual uint32				NumNestedScopes() const override final;
+	PLAINPROPS_API virtual uint32				NumParametricTypes() const override final;
+	PLAINPROPS_API virtual FNestedScope			Resolve(FNestedScopeId Id) const override final;
+	PLAINPROPS_API virtual FParametricTypeView	Resolve(FParametricTypeId Id) const override final;
+};
+
 } // namespace PlainProps
diff --git a/Engine/Plugins/Experimental/PlainProps/Source/Public/PlainPropsSave.h b/Engine/Plugins/Experimental/PlainProps/Source/Public/PlainPropsSave.h
--- a/Engine/Plugins/Experimental/PlainProps/Source/Public/PlainPropsSave.h
+++ b/Engine/Plugins/Experimental/PlainProps/Source/Public/PlainPropsSave.h
@@ -34,9 +34,9 @@ FSaveContext MakeSaveContext(FScratchAllocator& Scratch)
 
 ////////////////////////////////////////////////////////////////////////////////////////////////
 
-[[nodiscard]] PLAINPROPS_API FBuiltStruct*	SaveStruct(const void* Struct, FStructSchemaId BindId, const FSaveContext& Context);
-[[nodiscard]] PLAINPROPS_API FBuiltStruct*	SaveStructDelta(const void* Struct, const void* Default, FStructSchemaId BindId, const FSaveContext& Context);
-[[nodiscard]] PLAINPROPS_API FBuiltStruct*	SaveStructDeltaIfDiff(const void* Struct, const void* Default, FStructSchemaId BindId, const FSaveContext& Context);
+[[nodiscard]] PLAINPROPS_API FBuiltStruct*	SaveStruct(const void* Struct, FStructId BindId, const FSaveContext& Context);
+[[nodiscard]] PLAINPROPS_API FBuiltStruct*	SaveStructDelta(const void* Struct, const void* Default, FStructId BindId, const FSaveContext& Context);
+[[nodiscard]] PLAINPROPS_API FBuiltStruct*	SaveStructDeltaIfDiff(const void* Struct, const void* Default, FStructId BindId, const FSaveContext& Context);
 [[nodiscard]] PLAINPROPS_API FBuiltRange*	SaveRange(const void* Range, FRangeMemberBinding Member, const FSaveContext& Ctx);
 [[nodiscard]] PLAINPROPS_API FBuiltRange*	SaveLeafRange(const void* Range, const ILeafRangeBinding& Binding, FUnpackedLeafType Leaf, const FSaveContext& Ctx);
 
diff --git a/Engine/Plugins/Experimental/PlainProps/Source/Public/PlainPropsSaveMember.h b/Engine/Plugins/Experimental/PlainProps/Source/Public/PlainPropsSaveMember.h
--- a/Engine/Plugins/Experimental/PlainProps/Source/Public/PlainPropsSaveMember.h
+++ b/Engine/Plugins/Experimental/PlainProps/Source/Public/PlainPropsSaveMember.h
@@ -18,7 +18,7 @@ struct TStructSerializer
 	static constexpr EMemberKind Kind = EMemberKind::Struct;
 
 	template<class Ids>
-	TStructSerializer(TCustomInit<Ids> Init, TConstArrayView<FMemberId> Names, FDualStructSchemaId Dual = GetStructDualId<Ids, T>())
+	TStructSerializer(TCustomInit<Ids> Init, TConstArrayView<FMemberId> Names, FDualStructId Dual = GetStructDualId<Ids, T>())
 	: BindId(Dual.BindId)
 	{
 		Init.RegisterInnerStruct(Dual, Names);
@@ -26,7 +26,7 @@ struct TStructSerializer
 
 	FTypedValue SaveMember(const T& Value, const FSaveContext& Ctx) const
 	{
-		FMemberSchema Schema = { DefaultStructType, DefaultStructType, 0, BindId, nullptr};
+		FMemberSchema Schema = { DefaultStructType, DefaultStructType, 0, FInnerId(BindId), nullptr};
 		return { Schema, {.Struct = SaveItem(Value, Ctx)} };
 	}
 
@@ -59,10 +59,10 @@ struct TStructSerializer
 
 	FMemberSchema MakeMemberRangeSchema(ERangeSizeType MaxSize) const
 	{
-		return { FMemberType(MaxSize), DefaultStructType, 1, BindId, nullptr };
+		return { FMemberType(MaxSize), DefaultStructType, 1, FInnerId(BindId), nullptr };
 	}
 
-	const FStructSchemaId BindId;
+	const FStructId BindId;
 };
 
 //////////////////////////////////////////////////////////////////////////
@@ -101,7 +101,7 @@ struct TRangeSerializer
 	static constexpr uint16			NumTypes = NumRanges + 1;
 
 	FRangeBinding const* const		Bindings = nullptr;
-	FOptionalSchemaId				InnermostSchema;
+	FOptionalInnerId				InnermostSchema;
 	FMemberBindTypeInitializer		BindTypes[NumTypes];
 	FMemberType						Types[NumTypes];
 
@@ -123,7 +123,7 @@ struct TRangeSerializer
 
 		if constexpr (!LeafType<InnermostType>)
 		{			
-			FDualStructSchemaId Id = {static_cast<FStructSchemaId>(InnermostSchema.Get()), GetStructDeclId<Ids, InnermostType>()};
+			FDualStructId Id = {InnermostSchema.Get().AsStruct(), GetStructDeclId<Ids, InnermostType>()};
 			Init.RegisterInnerStruct(Id, Names);
 		}
 	}
@@ -183,7 +183,7 @@ struct TLeafSerializer
 		LoadMember(Dst, Src);
 	}
 
-	static FTypedValue SaveLeaf(FMemberType Type, FOptionalSchemaId Id, T Value)
+	static FTypedValue SaveLeaf(FMemberType Type, FOptionalInnerId Id, T Value)
 	{
 		return { {Type, Type, 0, Id, nullptr}, {.Leaf = ValueCast(Value)} };
 	}
@@ -224,7 +224,7 @@ struct TEnumSerializer : TLeafSerializer<T>
 {
 	static constexpr FMemberType MemberType = ReflectEnum<T>.Pack();
 	
-	const FEnumSchemaId Id;
+	const FEnumId Id;
 
 	template<class Ids>
 	TEnumSerializer(TCustomInit<Ids>, TConstArrayView<FMemberId>)
@@ -240,7 +240,7 @@ struct TEnumSerializer : TLeafSerializer<T>
 
 	FMemberSchema MakeMemberRangeSchema(ERangeSizeType MaxSize) const
 	{
-		return { FMemberType(MaxSize), MemberType, 1, Id, nullptr };
+		return { FMemberType(MaxSize), MemberType, 1, FInnerId(Id), nullptr };
 	}
 };
 
diff --git a/Engine/Plugins/Experimental/PlainProps/Source/Public/PlainPropsTypes.h b/Engine/Plugins/Experimental/PlainProps/Source/Public/PlainPropsTypes.h
--- a/Engine/Plugins/Experimental/PlainProps/Source/Public/PlainPropsTypes.h
+++ b/Engine/Plugins/Experimental/PlainProps/Source/Public/PlainPropsTypes.h
@@ -123,17 +123,53 @@ struct FMemberId
 
 //////////////////////////////////////////////////////////////////////////
 
+// Runtime id for an enum FType
+struct FEnumId
+{
+	uint32 Idx = ~0u;
+	friend uint32 GetTypeHash(FEnumId In) { return GetTypeHash(In.Idx); };
+	friend bool operator==(FEnumId, FEnumId) = default;
+};
+
+// Runtime id for a struct FType
+struct FStructId
+{
+	uint32 Idx = ~0u;
+	friend uint32 GetTypeHash(FStructId In) { return GetTypeHash(In.Idx); };
+	friend bool operator==(FStructId, FStructId) = default;
+};
+
+// Either a runtime FStructId or FEnumId of a member
+struct FInnerId
+{
+	uint32 Idx = ~0u;
+
+	explicit FInnerId(FEnumId In) : Idx(In.Idx) {}
+	explicit FInnerId(FStructId In) : Idx(In.Idx) {}
+	
+	FEnumId		AsEnum() const		{ return {Idx}; }
+	FStructId	AsStruct() const	{ return {Idx}; }
+
+	friend uint32 GetTypeHash(FInnerId In) { return GetTypeHash(In.Idx); };
+	friend bool operator==(FInnerId, FInnerId) = default;
+	explicit FInnerId(uint32 InIdx) : Idx(InIdx) {}
+};
+
+//////////////////////////////////////////////////////////////////////////
+
+// Serialized schema id of a runtime FInnerId
 struct FSchemaId
 {
 	uint32 Idx = ~0u;
 	
-	// TODO: Consider moving these to FEnumSchemaId/FStructSchemaId.
-	// They're not comparable in the building stage when the overlap.
-	// After remap they are comparable during reading.
 	bool operator==(FSchemaId O) const { return Idx == O.Idx; }
 	friend uint32 GetTypeHash(FSchemaId Schema) { return GetTypeHash(Schema.Idx); };
 };
+
+// Serialized struct id of a runtime FStructId
 struct FStructSchemaId : FSchemaId {};
+
+// Serialized enum id of a runtime FEnumId
 struct FEnumSchemaId : FSchemaId {};
 
 //////////////////////////////////////////////////////////////////////////
@@ -221,21 +257,24 @@ struct FType
 
 //////////////////////////////////////////////////////////////////////////
 
+inline uint32 ToIdx(FInnerId Id) { return Id.Idx; }
 inline uint32 ToIdx(FNameId Id) { return Id.Idx; }
 inline uint32 ToIdx(FMemberId Name) { return Name.Id.Idx; }
+inline uint32 ToIdx(FEnumId Id) { return Id.Idx; }
+inline uint32 ToIdx(FStructId Id) { return Id.Idx; }
 inline uint32 ToIdx(FSchemaId Id) { return Id.Idx; }
 inline uint32 ToIdx(FNestedScopeId Id) { return Id.Idx; }
 inline uint32 ToIdx(FParametricTypeId Id) { return Id.AsInt(); }
 inline uint32 ToIdx(FConcreteTypenameId Name) { return Name.Id.Idx; }
 
 template<class IdType>
-IdType FromIdx(uint32 Idx) { return {Idx}; }
-
-template<>
-inline FMemberId FromIdx(uint32 Idx) { return {{Idx}}; }
-
-template<>
-inline FParametricTypeId FromIdx(uint32 Idx) { return FParametricTypeId::FromInt(Idx); }
+IdType FromIdx(uint32 Idx)
+{
+	return {Idx};
+}
+template<> inline FMemberId FromIdx(uint32 Idx) { return {{Idx}}; }
+template<> inline FInnerId FromIdx(uint32 Idx) { return FInnerId(Idx); }
+template<> inline FParametricTypeId FromIdx(uint32 Idx) { return FParametricTypeId::FromInt(Idx); }
 
 class FOptionalId
 {
@@ -253,11 +292,15 @@ public:
 	constexpr TOptionalId() = default;
 	constexpr TOptionalId(FNoId) {}
     constexpr TOptionalId(T Id) { Idx = ToIdx(Id); }
+
+	template<typename U>
+	constexpr explicit TOptionalId(TOptionalId<U> In) requires (std::is_constructible_v<T,U>)
+	: FOptionalId(static_cast<const FOptionalId&>(In))
+	{}
     
 	template<class U>
-    explicit operator TOptionalId<U>() const
+    explicit operator TOptionalId<U>() const requires(std::is_convertible_v<U, T> || std::is_convertible_v<T, U>)
     {
-		static_assert(std::is_convertible_v<U, T> || std::is_convertible_v<T, U>);
 		return static_cast<const TOptionalId<U>&>(static_cast<const FOptionalId&>(*this));
     }
 
@@ -271,9 +314,11 @@ public:
 	friend bool operator!=(TOptionalId A, TOptionalId B) { return A.Idx != B.Idx;}
 };
 
-
 using FOptionalNameId = TOptionalId<FNameId>;
 using FOptionalMemberId = TOptionalId<FMemberId>;
+using FOptionalEnumId = TOptionalId<FEnumId>;
+using FOptionalStructId = TOptionalId<FStructId>;
+using FOptionalInnerId = TOptionalId<FInnerId>;
 using FOptionalSchemaId = TOptionalId<FSchemaId>;
 using FOptionalStructSchemaId = TOptionalId<FStructSchemaId>;
 using FOptionalEnumSchemaId = TOptionalId<FEnumSchemaId>;
@@ -284,6 +329,11 @@ using FOptionalConcreteTypenameId = TOptionalId<FConcreteTypenameId>;
 template<class IdType>
 inline constexpr TOptionalId<IdType> ToOptional(IdType Id) { return Id; }
 
+inline FOptionalEnumId			ToOptionalEnum(FOptionalInnerId In)		{ return static_cast<const FOptionalEnumId&>(static_cast<const FOptionalId&>(In)); }
+inline FOptionalStructId		ToOptionalStruct(FOptionalInnerId In)	{ return static_cast<const FOptionalStructId&>(static_cast<const FOptionalId&>(In)); }
+inline FOptionalEnumSchemaId	ToOptionalEnum(FOptionalSchemaId In)	{ return static_cast<const FOptionalEnumSchemaId&>(static_cast<const FOptionalId&>(In)); }
+inline FOptionalStructSchemaId	ToOptionalStruct(FOptionalSchemaId In)	{ return static_cast<const FOptionalStructSchemaId&>(static_cast<const FOptionalId&>(In)); }
+
 //////////////////////////////////////////////////////////////////////////
 
 // Resolved FNestedScopeId
@@ -420,48 +470,59 @@ enum class ESchemaFormat { StableNames, InMemoryNames };
 
 using FUtf8Builder = FUtf8StringBuilderBase;
 
-class FIds
+// Resolves structured ids and converts ids to strings
+class FIdsBase
 {
 public:
-	virtual ~FIds() {}
-
-	virtual uint32						NumNames() const = 0;
-	virtual uint32						NumNestedScopes() const = 0;
-	virtual uint32						NumParametricTypes() const = 0;
-	virtual uint32						NumEnums() const = 0;
-	virtual uint32						NumStructs() const = 0;
-
-	virtual FNestedScope				Resolve(FNestedScopeId Id) const = 0;
-	virtual FParametricTypeView			Resolve(FParametricTypeId Id) const = 0;
-	virtual FType						Resolve(FEnumSchemaId Id) const = 0;
-	virtual FType						Resolve(FStructSchemaId Id) const = 0;
-
-	virtual void						AppendString(FUtf8Builder& Out, FNameId Name) const = 0;
-
-	PLAINPROPS_API virtual void			AppendString(FUtf8Builder& Out, FMemberId Name) const;
-	PLAINPROPS_API virtual void			AppendString(FUtf8Builder& Out, FOptionalMemberId Name) const;
-	PLAINPROPS_API virtual void			AppendString(FUtf8Builder& Out, FScopeId Scope) const;
-	PLAINPROPS_API virtual void			AppendString(FUtf8Builder& Out, FTypenameId Typename) const;
-	PLAINPROPS_API virtual void			AppendString(FUtf8Builder& Out, FType Type) const;
-	PLAINPROPS_API virtual void			AppendString(FUtf8Builder& Out, FEnumSchemaId Enum) const;
-	PLAINPROPS_API virtual void			AppendString(FUtf8Builder& Out, FStructSchemaId Struct) const;
+	virtual ~FIdsBase() {}
+
+	virtual uint32					NumNames() const = 0;
+	virtual uint32					NumNestedScopes() const = 0;
+	virtual uint32					NumParametricTypes() const = 0;
+
+	virtual FNestedScope			Resolve(FNestedScopeId Id) const = 0;
+	virtual FParametricTypeView		Resolve(FParametricTypeId Id) const = 0;
+
+	virtual void					AppendString(FUtf8Builder& Out, FNameId Name) const = 0;
+	PLAINPROPS_API virtual void		AppendString(FUtf8Builder& Out, FMemberId Name) const;
+	PLAINPROPS_API virtual void		AppendString(FUtf8Builder& Out, FOptionalMemberId Name) const;
+	PLAINPROPS_API virtual void		AppendString(FUtf8Builder& Out, FScopeId Scope) const;
+	PLAINPROPS_API virtual void		AppendString(FUtf8Builder& Out, FTypenameId Typename) const;
+	PLAINPROPS_API virtual void		AppendString(FUtf8Builder& Out, FType Type) const;
+};
+
+// Runtime id resolver
+class FIds : public FIdsBase
+{
+public:
+	using FIdsBase::Resolve;
+	using FIdsBase::AppendString;
+
+	virtual uint32					NumEnums() const = 0;
+	virtual uint32					NumStructs() const = 0;
+	virtual FType					Resolve(FEnumId Id) const = 0;
+	virtual FType					Resolve(FStructId Id) const = 0;
+
+	PLAINPROPS_API virtual void		AppendString(FUtf8Builder& Out, FEnumId Enum) const;
+	PLAINPROPS_API virtual void		AppendString(FUtf8Builder& Out, FStructId Struct) const;
 };
 
+// Helps format log messages
 class FDebugIds
 {
 public:
 	PLAINPROPS_API explicit	FDebugIds(const FIds& InIds) : Ids(InIds) {}
 
-	PLAINPROPS_API FString				Print(FNameId Name) const;
-	PLAINPROPS_API FString				Print(FMemberId Name) const;
-	PLAINPROPS_API FString				Print(FOptionalMemberId Name) const;
-	PLAINPROPS_API FString				Print(FScopeId Scope) const;
-	PLAINPROPS_API FString				Print(FTypenameId Typename) const;
-	PLAINPROPS_API FString				Print(FConcreteTypenameId Typename) const;
-	PLAINPROPS_API FString				Print(FParametricTypeId Typename) const;
-	PLAINPROPS_API FString				Print(FType Type) const;
-	PLAINPROPS_API FString				Print(FEnumSchemaId Enum) const;
-	PLAINPROPS_API FString				Print(FStructSchemaId Struct) const;
+	PLAINPROPS_API FString			Print(FNameId Name) const;
+	PLAINPROPS_API FString			Print(FMemberId Name) const;
+	PLAINPROPS_API FString			Print(FOptionalMemberId Name) const;
+	PLAINPROPS_API FString			Print(FScopeId Scope) const;
+	PLAINPROPS_API FString			Print(FTypenameId Typename) const;
+	PLAINPROPS_API FString			Print(FConcreteTypenameId Typename) const;
+	PLAINPROPS_API FString			Print(FParametricTypeId Typename) const;
+	PLAINPROPS_API FString			Print(FType Type) const;
+	PLAINPROPS_API FString			Print(FEnumId Enum) const;
+	PLAINPROPS_API FString			Print(FStructId Struct) const;
 private:
 	const FIds& Ids;
 };
diff --git a/Engine/Plugins/Experimental/PlainProps/Source/Public/PlainPropsUeCoreBindings.h b/Engine/Plugins/Experimental/PlainProps/Source/Public/PlainPropsUeCoreBindings.h
--- a/Engine/Plugins/Experimental/PlainProps/Source/Public/PlainPropsUeCoreBindings.h
+++ b/Engine/Plugins/Experimental/PlainProps/Source/Public/PlainPropsUeCoreBindings.h
@@ -504,8 +504,8 @@ struct FTransformBinding : ICustomBinding
 	using Type = FTransform;
 	enum class EMember : uint8 { Translate, Rotate, Scale };
 	const FMemberId MemberIds[3];
-	const FStructSchemaId VectorId;
-	const FStructSchemaId QuatId;
+	const FStructId VectorId;
+	const FStructId QuatId;
 
 	template<class Ids>
 	FTransformBinding(TCustomInit<Ids>)
diff --git a/Engine/Plugins/Experimental/PlainProps/Source/Public/PlainPropsWrite.h b/Engine/Plugins/Experimental/PlainProps/Source/Public/PlainPropsWrite.h
--- a/Engine/Plugins/Experimental/PlainProps/Source/Public/PlainPropsWrite.h
+++ b/Engine/Plugins/Experimental/PlainProps/Source/Public/PlainPropsWrite.h
@@ -25,10 +25,10 @@ public:
 	PLAINPROPS_API ~FWriter();
 	
 	PLAINPROPS_API TConstArrayView<FNameId>			GetUsedNames() const;
-	PLAINPROPS_API FOptionalStructSchemaId			GetWriteId(FStructSchemaId BuiltId) const;
+	PLAINPROPS_API FOptionalStructSchemaId			GetWriteId(FStructId BuiltId) const;
 
 	PLAINPROPS_API void								WriteSchemas(TArray64<uint8>& Out) const;
-	PLAINPROPS_API FStructSchemaId					WriteMembers(TArray64<uint8>& Out, FStructSchemaId BuiltId, const FBuiltStruct& Struct) const;
+	PLAINPROPS_API FStructSchemaId					WriteMembers(TArray64<uint8>& Out, FStructId BuiltId, const FBuiltStruct& Struct) const;
 
 private:
 	const FBuiltSchemas&							Schemas;
diff --git a/Engine/Plugins/Experimental/PlainProps/Source/Tests/PlainPropsTests.cpp b/Engine/Plugins/Experimental/PlainProps/Source/Tests/PlainPropsTests.cpp
--- a/Engine/Plugins/Experimental/PlainProps/Source/Tests/PlainPropsTests.cpp
+++ b/Engine/Plugins/Experimental/PlainProps/Source/Tests/PlainPropsTests.cpp
@@ -266,25 +266,25 @@ class FTestBatchBuilder : public TIdIndexer<FAnsiString>, public IStructBindIds
 public:
 	FTestBatchBuilder(FScratchAllocator& InScratch) : Declarations(FDebugIds(*this)), Scratch(InScratch) {}
 
-	FEnumSchemaId				DeclareEnum(FType Type, EEnumMode Mode, ELeafWidth Width, std::initializer_list<const char*> Names, std::initializer_list<uint64> Constants);
-	FEnumSchemaId				DeclareEnum(const char* Scope, const char* Name, EEnumMode Mode, ELeafWidth Width, std::initializer_list<const char*> Names, std::initializer_list<uint64> Constants);
-	FStructSchemaId				DeclareStruct(FType Type, std::initializer_list<const char*> MemberOrder, EMemberPresence Occupancy, FOptionalStructSchemaId Super = {});
-	FStructSchemaId				DeclareStruct(const char* Scope, const char* Name, std::initializer_list<const char*> MemberOrder, EMemberPresence Occupancy, FOptionalStructSchemaId Super = {});
+	FEnumId						DeclareEnum(FType Type, EEnumMode Mode, ELeafWidth Width, std::initializer_list<const char*> Names, std::initializer_list<uint64> Constants);
+	FEnumId						DeclareEnum(const char* Scope, const char* Name, EEnumMode Mode, ELeafWidth Width, std::initializer_list<const char*> Names, std::initializer_list<uint64> Constants);
+	FStructId					DeclareStruct(FType Type, std::initializer_list<const char*> MemberOrder, EMemberPresence Occupancy, FOptionalStructId Super = {});
+	FStructId					DeclareStruct(const char* Scope, const char* Name, std::initializer_list<const char*> MemberOrder, EMemberPresence Occupancy, FOptionalStructId Super = {});
 	
-	const FEnumDeclaration&		Get(FEnumSchemaId Id) { return Declarations.Get(Id); }
-	const FStructDeclaration&	Get(FStructSchemaId Id) { return Declarations.Get(Id); }
+	const FEnumDeclaration&		Get(FEnumId Id) { return Declarations.Get(Id); }
+	const FStructDeclaration&	Get(FStructId Id) { return Declarations.Get(Id); }
 
-	void						AddObject(FStructSchemaId Schema, FMemberBuilder&& Members);
-	void						AddObjects(FStructSchemaId Schema, std::initializer_list<FMemberBuilder&&> Objects);
+	void						AddObject(FStructId Schema, FMemberBuilder&& Members);
+	void						AddObjects(FStructId Schema, std::initializer_list<FMemberBuilder&&> Objects);
 
 	TArray64<uint8>				Write();
 
 	FDebugIds					GetDebug() const { return FDebugIds(*this); }
 
 private:
-	TArray<TPair<FStructSchemaId, FBuiltStruct*>>	Objects;
-	FDeclarations									Declarations;
-	FScratchAllocator&								Scratch;
+	TArray<TPair<FStructId, FBuiltStruct*>>		Objects;
+	FDeclarations								Declarations;
+	FScratchAllocator&							Scratch;
 
 	TArray<FMemberId> NameMembers(std::initializer_list<const char*> Members)
 	{
@@ -311,38 +311,38 @@ private:
 
 	TArray<char> GetNameData() const;
 
-	virtual FStructSchemaId GetDeclId(FStructSchemaId BindId) const override final
+	virtual FStructId GetDeclId(FStructId BindId) const override final
 	{
 		checkf(false, TEXT("All struct ids should be declared, nothing is bound with different names in this test suite"));
 		return BindId;
 	}
 };
 
-FStructSchemaId FTestBatchBuilder::DeclareStruct(FType Type, std::initializer_list<const char*> MemberOrder, EMemberPresence Occupancy, FOptionalStructSchemaId Super)
+FStructId FTestBatchBuilder::DeclareStruct(FType Type, std::initializer_list<const char*> MemberOrder, EMemberPresence Occupancy, FOptionalStructId Super)
 {
-	FStructSchemaId Id = IndexStruct(Type);
+	FStructId Id = IndexStruct(Type);
 	Declarations.DeclareStruct(Id, Type, NameMembers(MemberOrder), Occupancy, Super);
 	return Id;
 }
 
-FStructSchemaId FTestBatchBuilder::DeclareStruct(const char* Scope, const char* Name, std::initializer_list<const char*> MemberOrder, EMemberPresence Occupancy, FOptionalStructSchemaId Super)
+FStructId FTestBatchBuilder::DeclareStruct(const char* Scope, const char* Name, std::initializer_list<const char*> MemberOrder, EMemberPresence Occupancy, FOptionalStructId Super)
 {
 	return DeclareStruct(MakeType(Scope, Name), MemberOrder, Occupancy, Super);
 }
 
-FEnumSchemaId FTestBatchBuilder::DeclareEnum(FType Type, EEnumMode Mode, ELeafWidth, std::initializer_list<const char*> InNames, std::initializer_list<uint64> Constants)
+FEnumId FTestBatchBuilder::DeclareEnum(FType Type, EEnumMode Mode, ELeafWidth, std::initializer_list<const char*> InNames, std::initializer_list<uint64> Constants)
 {
-	FEnumSchemaId Id = IndexEnum(Type);
+	FEnumId Id = IndexEnum(Type);
 	Declarations.DeclareEnum(Id, Type, Mode, MakeEnumerators(InNames, Constants), EEnumAliases::Fail);
 	return Id;
 }
 
-FEnumSchemaId FTestBatchBuilder::DeclareEnum(const char* Scope, const char* Name, EEnumMode Mode, ELeafWidth Width, std::initializer_list<const char*> InNames, std::initializer_list<uint64> Constants)
+FEnumId FTestBatchBuilder::DeclareEnum(const char* Scope, const char* Name, EEnumMode Mode, ELeafWidth Width, std::initializer_list<const char*> InNames, std::initializer_list<uint64> Constants)
 {
 	return DeclareEnum(MakeType(Scope, Name), Mode, Width, InNames, Constants);
 }
 
-void FTestBatchBuilder::AddObject(FStructSchemaId Schema, FMemberBuilder&& Members)
+void FTestBatchBuilder::AddObject(FStructId Schema, FMemberBuilder&& Members)
 {
 	Objects.Emplace(Schema, Members.BuildAndReset(Scratch, Declarations.Get(Schema), GetDebug()));
 }
@@ -351,7 +351,7 @@ TArray64<uint8> FTestBatchBuilder::Write()
 {
 	// Build partial schemas
 	FSchemasBuilder SchemaBuilders(Declarations, *this, *this, Scratch, ESchemaFormat::StableNames);
-	for (const TPair<FStructSchemaId, FBuiltStruct*>& Object : Objects)
+	for (const TPair<FStructId, FBuiltStruct*>& Object : Objects)
 	{
 		SchemaBuilders.NoteStructAndMembers(Object.Key, *Object.Value);
 	}
@@ -382,7 +382,7 @@ TArray64<uint8> FTestBatchBuilder::Write()
 
 	// Write objects
 	WriteInt(Out, TestMagics[2]);
-	for (const TPair<FStructSchemaId, FBuiltStruct*>& Object : Objects)
+	for (const TPair<FStructId, FBuiltStruct*>& Object : Objects)
 	{
 		WriteInt(/* out */ Tmp, TestMagics[3]);
 		WriteInt(/* out */ Tmp, Writer.GetWriteId(Object.Key).Get().Idx);
@@ -445,28 +445,15 @@ private:
 	TArray<const char*> Names;
 };
 
-class FTestBatchNames : public FIds
+class FTestBatchNames final : public FStableBatchIds
 {
-public:
-	FTestBatchNames(const FTestNameReader& InNames, const FSchemaBatch& InSchemas) : Names(InNames), Schemas(InSchemas) {}
-
-	virtual uint32					NumNames() const override final						{ return Names.NumNames(); }
-	virtual uint32					NumNestedScopes() const override final				{ return Schemas.NumNestedScopes; }
-	virtual uint32					NumParametricTypes() const override final			{ return Schemas.NumParametricTypes; }
-	virtual uint32					NumEnums() const override final						{ return Schemas.NumSchemas - Schemas.NumStructSchemas; }
-	virtual uint32					NumStructs() const override final					{ return Schemas.NumStructSchemas; }
-
-	virtual FNestedScope			Resolve(FNestedScopeId Id) const override final		{ return ResolveNestedScope(Schemas, Id); }
-	virtual FParametricTypeView		Resolve(FParametricTypeId Id) const override final	{ return ResolveParametricType(Schemas, Id); }
-	virtual FType					Resolve(FEnumSchemaId Id) const override final		{ return ResolveEnumSchema(Schemas, Id).Type; }
-	virtual FType					Resolve(FStructSchemaId Id) const override final	{ return ResolveStructSchema(Schemas, Id).Type; }
-
-	using FIds::AppendString;
-	virtual void AppendString(FUtf8StringBuilderBase& Out, FNameId Name) const override final { Out.Append(Names[Name]); }
-
-private:
 	const FTestNameReader& Names;
-	const FSchemaBatch& Schemas;
+public:
+	FTestBatchNames(const FTestNameReader& InNames, FReadBatchId Batch) : FStableBatchIds(Batch), Names(InNames) {}
+	
+	using FStableBatchIds::AppendString;
+	virtual uint32 NumNames() const override { return Names.NumNames(); }
+	virtual void AppendString(FUtf8StringBuilderBase& Out, FNameId Name) const override { Out.Append(Names[Name]); }
 };
 
 static void PrintBatch(
@@ -476,8 +463,8 @@ static void PrintBatch(
 	const FSchemaBatch& Schemas,
 	TConstArrayView<FStructView> Objects)
 {
-	FTestBatchNames BatchNames(Names, Schemas);
-	PrintYamlBatch(Out, BatchNames, Batch, Objects);
+	FTestBatchNames BatchNames(Names, Batch);
+	PrintYamlBatch(Out, BatchNames, Objects);
 	FUtf8StringView YamlView = Out.ToView();
 	auto Wide = StringCast<TCHAR>(YamlView.GetData(), YamlView.Len());
 	UE_LOG(LogPlainPropsTests, Log, TEXT("\n%.*s\n"), Wide.Length(), Wide.Get());
@@ -596,7 +583,7 @@ TEST_CASE_NAMED(FPlainPropsReadWriteTest, "System::Core::Serialization::PlainPro
 
 		TestSerialize([](FTestBatchBuilder& Batch, FScratchAllocator& Scratch)
 		{
-			FStructSchemaId SchemaId = Batch.DeclareStruct("Testing", "Bools", MemberNames, EMemberPresence::AllowSparse);
+			FStructId SchemaId = Batch.DeclareStruct("Testing", "Bools", MemberNames, EMemberPresence::AllowSparse);
 
 			FMemberBuilder B1T;
 			B1T.Add(Batch.NameMember("b3"),	true);
@@ -689,7 +676,7 @@ TEST_CASE_NAMED(FPlainPropsReadWriteTest, "System::Core::Serialization::PlainPro
 
 		TestSerialize([](FTestBatchBuilder& Batch, FScratchAllocator& Scratch)
 		{
-			FStructSchemaId SchemaId = Batch.DeclareStruct("Test", "Numbers", MemberNames, EMemberPresence::AllowSparse);
+			FStructId SchemaId = Batch.DeclareStruct("Test", "Numbers", MemberNames, EMemberPresence::AllowSparse);
 
 			FMemberBuilder Misc, Mins, Maxs, Some;
 
@@ -796,8 +783,8 @@ TEST_CASE_NAMED(FPlainPropsReadWriteTest, "System::Core::Serialization::PlainPro
 
 		TestSerialize([](FTestBatchBuilder& Batch, FScratchAllocator& Scratch)
 		{
-			FStructSchemaId ExplicitId = Batch.DeclareStruct("Test", "ExplicitDense", ExplicitMembers, EMemberPresence::RequireAll);
-			FStructSchemaId ImplicitId = Batch.DeclareStruct("Test", "ImplicitDense", ImplicitMembers, EMemberPresence::AllowSparse);
+			FStructId ExplicitId = Batch.DeclareStruct("Test", "ExplicitDense", ExplicitMembers, EMemberPresence::RequireAll);
+			FStructId ImplicitId = Batch.DeclareStruct("Test", "ImplicitDense", ImplicitMembers, EMemberPresence::AllowSparse);
 
 			FMemberBuilder X;
 			X.Add(Batch.NameMember("A"), char8_t('a'));
@@ -855,11 +842,11 @@ TEST_CASE_NAMED(FPlainPropsReadWriteTest, "System::Core::Serialization::PlainPro
 
 		TestSerialize([](FTestBatchBuilder& Batch, FScratchAllocator& Scratch)
 		{
-			FStructSchemaId ObjectId = Batch.DeclareStruct("Test", "Object", ObjectMembers, EMemberPresence::AllowSparse);
-			FStructSchemaId StructId = Batch.DeclareStruct("Test", "Struct", StructMembers, EMemberPresence::AllowSparse);
-			FStructSchemaId NestedId = Batch.DeclareStruct("Test", "Nested", NestedMembers, EMemberPresence::AllowSparse);
-			FStructSchemaId UnusedId = Batch.DeclareStruct("Test", "Unused", UnusedMembers, EMemberPresence::AllowSparse);
-			FStructSchemaId EmptyId  = Batch.DeclareStruct("Test", "Empty",  EmptyMembers,  EMemberPresence::AllowSparse);
+			FStructId ObjectId = Batch.DeclareStruct("Test", "Object", ObjectMembers, EMemberPresence::AllowSparse);
+			FStructId StructId = Batch.DeclareStruct("Test", "Struct", StructMembers, EMemberPresence::AllowSparse);
+			FStructId NestedId = Batch.DeclareStruct("Test", "Nested", NestedMembers, EMemberPresence::AllowSparse);
+			FStructId UnusedId = Batch.DeclareStruct("Test", "Unused", UnusedMembers, EMemberPresence::AllowSparse);
+			FStructId EmptyId  = Batch.DeclareStruct("Test", "Empty",  EmptyMembers,  EMemberPresence::AllowSparse);
 
 			FMemberBuilder Members;
 			Members.Add(Batch.NameMember("I1"), 100);
@@ -916,21 +903,21 @@ TEST_CASE_NAMED(FPlainPropsReadWriteTest, "System::Core::Serialization::PlainPro
 
 		TestSerialize([](FTestBatchBuilder& Batch, FScratchAllocator& Scratch)
 		{
-			// Test create holes in the original FNameId, FStructSchemaId and FEnumSchemaId index range
-			FStructSchemaId UnusedId = Batch.DeclareStruct("Test", "UnusedStruct", {"U1", "U2"}, EMemberPresence::AllowSparse);
-			FStructSchemaId ObjectId = Batch.DeclareStruct("Test", "Enums", MemberNames, EMemberPresence::AllowSparse);
-
-			FEnumSchemaId U	= Batch.DeclareEnum("Test", "UnusedEnum1",	EEnumMode::Flag, ELeafWidth::B8,	{"U3"}, {1}); // Hole
-			FEnumSchemaId A	= Batch.DeclareEnum("Test", "FlatDense8",	EEnumMode::Flat, ELeafWidth::B8,	{"A", "B", "C"}, {0, 1, 2});
-			FEnumSchemaId X	= Batch.DeclareEnum("Test", "UnusedEnum2",	EEnumMode::Flag, ELeafWidth::B8,	{"U4"}, {1}); // Hole
-			FEnumSchemaId B	= Batch.DeclareEnum("Test", "FlagDense8",	EEnumMode::Flag, ELeafWidth::B8,	{"A", "B", "C"}, {1, 2, 4});
-			FEnumSchemaId C	= Batch.DeclareEnum("Test", "FlatSparse8",	EEnumMode::Flat, ELeafWidth::B8,	{"A", "B", "C"}, {1, 2, 3});
-			FEnumSchemaId D	= Batch.DeclareEnum("Test", "FlagSparse8",	EEnumMode::Flag, ELeafWidth::B8,	{"A", "B", "C"}, {2, 16, 32});
-			FEnumSchemaId E	= Batch.DeclareEnum("Test", "FlatLimit8",	EEnumMode::Flat, ELeafWidth::B8,	{"Min", "Max"}, {0, 0xFF});
-			FEnumSchemaId F	= Batch.DeclareEnum("Test", "FlatLimit16",	EEnumMode::Flat, ELeafWidth::B16,	{"Min", "Max"}, {0, 0xFFFF});
-			FEnumSchemaId G	= Batch.DeclareEnum("Test", "FlatLimit32",	EEnumMode::Flat, ELeafWidth::B32,	{"Min", "Max"}, {0, 0xFFFFFFFF});
-			FEnumSchemaId H	= Batch.DeclareEnum("Test", "FlatLimit64",	EEnumMode::Flat, ELeafWidth::B64,	{"Min", "Max"}, {0, 0xFFFFFFFFFFFFFFFF});
-			FEnumSchemaId I	= Batch.DeclareEnum("Test", "FlagLimit64",	EEnumMode::Flag, ELeafWidth::B64,	{"One", "Max"}, {1, 0x8000000000000000});
+			// Test create holes in the original FNameId, FStructId and FEnumId index range
+			FStructId UnusedId = Batch.DeclareStruct("Test", "UnusedStruct", {"U1", "U2"}, EMemberPresence::AllowSparse);
+			FStructId ObjectId = Batch.DeclareStruct("Test", "Enums", MemberNames, EMemberPresence::AllowSparse);
+
+			FEnumId U	= Batch.DeclareEnum("Test", "UnusedEnum1",	EEnumMode::Flag, ELeafWidth::B8,	{"U3"}, {1}); // Hole
+			FEnumId A	= Batch.DeclareEnum("Test", "FlatDense8",	EEnumMode::Flat, ELeafWidth::B8,	{"A", "B", "C"}, {0, 1, 2});
+			FEnumId X	= Batch.DeclareEnum("Test", "UnusedEnum2",	EEnumMode::Flag, ELeafWidth::B8,	{"U4"}, {1}); // Hole
+			FEnumId B	= Batch.DeclareEnum("Test", "FlagDense8",	EEnumMode::Flag, ELeafWidth::B8,	{"A", "B", "C"}, {1, 2, 4});
+			FEnumId C	= Batch.DeclareEnum("Test", "FlatSparse8",	EEnumMode::Flat, ELeafWidth::B8,	{"A", "B", "C"}, {1, 2, 3});
+			FEnumId D	= Batch.DeclareEnum("Test", "FlagSparse8",	EEnumMode::Flag, ELeafWidth::B8,	{"A", "B", "C"}, {2, 16, 32});
+			FEnumId E	= Batch.DeclareEnum("Test", "FlatLimit8",	EEnumMode::Flat, ELeafWidth::B8,	{"Min", "Max"}, {0, 0xFF});
+			FEnumId F	= Batch.DeclareEnum("Test", "FlatLimit16",	EEnumMode::Flat, ELeafWidth::B16,	{"Min", "Max"}, {0, 0xFFFF});
+			FEnumId G	= Batch.DeclareEnum("Test", "FlatLimit32",	EEnumMode::Flat, ELeafWidth::B32,	{"Min", "Max"}, {0, 0xFFFFFFFF});
+			FEnumId H	= Batch.DeclareEnum("Test", "FlatLimit64",	EEnumMode::Flat, ELeafWidth::B64,	{"Min", "Max"}, {0, 0xFFFFFFFFFFFFFFFF});
+			FEnumId I	= Batch.DeclareEnum("Test", "FlagLimit64",	EEnumMode::Flag, ELeafWidth::B64,	{"One", "Max"}, {1, 0x8000000000000000});
 					
 			FMemberBuilder Members;
 			Members.AddEnum(Batch.NameMember("A2"),				A, (uint8)2);
@@ -1049,10 +1036,10 @@ TEST_CASE_NAMED(FPlainPropsReadWriteTest, "System::Core::Serialization::PlainPro
 		{
 			static constexpr std::initializer_list<const char*> MemberNames = { "B0", "B1", "B8", "B9", "D0", "D3", "Hi", "E3", "E0" };
 		
-			FStructSchemaId ObjectId = Batch.DeclareStruct("Test", "Object", MemberNames, EMemberPresence::AllowSparse);
-			FEnumSchemaId Enum	= Batch.DeclareEnum("Test", "ABCD", EEnumMode::Flat, ELeafWidth::B16, {"A", "B", "C", "D"}, {0, 1, 2, 3});
-			FEnumSchemaId UnusedEnum1	= Batch.DeclareEnum("Test", "Unused1", EEnumMode::Flat, ELeafWidth::B8, {"X"}, {0});
-			FEnumSchemaId UnusedEnum2	= Batch.DeclareEnum("Test", "Unused2", EEnumMode::Flat, ELeafWidth::B8, {"Y"}, {0});
+			FStructId ObjectId = Batch.DeclareStruct("Test", "Object", MemberNames, EMemberPresence::AllowSparse);
+			FEnumId Enum	= Batch.DeclareEnum("Test", "ABCD", EEnumMode::Flat, ELeafWidth::B16, {"A", "B", "C", "D"}, {0, 1, 2, 3});
+			FEnumId UnusedEnum1	= Batch.DeclareEnum("Test", "Unused1", EEnumMode::Flat, ELeafWidth::B8, {"X"}, {0});
+			FEnumId UnusedEnum2	= Batch.DeclareEnum("Test", "Unused2", EEnumMode::Flat, ELeafWidth::B8, {"Y"}, {0});
 		
 			FMemberBuilder Members;
 			Members.AddRange(Batch.NameMember("B0"), BuildLeafRange(Scratch, TConstArrayView<bool>()));
@@ -1098,8 +1085,8 @@ TEST_CASE_NAMED(FPlainPropsReadWriteTest, "System::Core::Serialization::PlainPro
 	{
 		TestSerialize([](FTestBatchBuilder& Batch, FScratchAllocator& Scratch)
 		{
-			FStructSchemaId ObjectId = Batch.DeclareStruct("Test", "Object", {"Structs"}, EMemberPresence::AllowSparse);
-			FStructSchemaId StructId = Batch.DeclareStruct("Test", "Struct", {"I", "F"}, EMemberPresence::AllowSparse);
+			FStructId ObjectId = Batch.DeclareStruct("Test", "Object", {"Structs"}, EMemberPresence::AllowSparse);
+			FStructId StructId = Batch.DeclareStruct("Test", "Struct", {"I", "F"}, EMemberPresence::AllowSparse);
 			 
 			FStructRangeBuilder Structs(4);
 			Structs[0].Add(Batch.NameMember("I"), 0);
@@ -1134,11 +1121,11 @@ TEST_CASE_NAMED(FPlainPropsReadWriteTest, "System::Core::Serialization::PlainPro
 
 		TestSerialize([](FTestBatchBuilder& Batch, FScratchAllocator& Scratch)
 		{
-			FStructSchemaId Object = Batch.DeclareStruct("Test", "Object", {"IntRs", "EmptyRs", "FloatRs", "EnumRs", "UnusedEnumRs", "UnicodeRs", "StructRs", "StructRRs"}, EMemberPresence::AllowSparse);
-			FStructSchemaId XY = Batch.DeclareStruct("Test", "XY", {"X", "Y"}, EMemberPresence::RequireAll);
-			FStructSchemaId ZW = Batch.DeclareStruct("Test", "ZW", {"Z", "W"}, EMemberPresence::AllowSparse);
-			FEnumSchemaId Enum = Batch.DeclareEnum("Test", "AB", EEnumMode::Flag, ELeafWidth::B8, {"A", "B"}, {1, 4});
-			FEnumSchemaId UnusedEnum = Batch.DeclareEnum("Test", "Unused", EEnumMode::Flat, ELeafWidth::B8, {"X"}, {0});
+			FStructId Object = Batch.DeclareStruct("Test", "Object", {"IntRs", "EmptyRs", "FloatRs", "EnumRs", "UnusedEnumRs", "UnicodeRs", "StructRs", "StructRRs"}, EMemberPresence::AllowSparse);
+			FStructId XY = Batch.DeclareStruct("Test", "XY", {"X", "Y"}, EMemberPresence::RequireAll);
+			FStructId ZW = Batch.DeclareStruct("Test", "ZW", {"Z", "W"}, EMemberPresence::AllowSparse);
+			FEnumId Enum = Batch.DeclareEnum("Test", "AB", EEnumMode::Flag, ELeafWidth::B8, {"A", "B"}, {1, 4});
+			FEnumId UnusedEnum = Batch.DeclareEnum("Test", "Unused", EEnumMode::Flat, ELeafWidth::B8, {"X"}, {0});
 	
 			FNestedRangeBuilder IntRs(MakeLeafRangeSchema<int32, int32>(), 3);
 			IntRs.Add(BuildLeafRange(Scratch, MakeArrayView({1})));
@@ -1257,8 +1244,8 @@ TEST_CASE_NAMED(FPlainPropsReadWriteTest, "System::Core::Serialization::PlainPro
 	{
 		TestSerialize([](FTestBatchBuilder& Batch, FScratchAllocator& Scratch)
 		{
-			FStructSchemaId Object = Batch.DeclareStruct("Test", "Object", {"Bools", "Structs", "BF", "BT" }, EMemberPresence::AllowSparse);
-			FStructSchemaId Struct = Batch.DeclareStruct("Test", "Struct", {"MaybeB", "Bs", "MaybeBs", "B"}, EMemberPresence::AllowSparse);
+			FStructId Object = Batch.DeclareStruct("Test", "Object", {"Bools", "Structs", "BF", "BT" }, EMemberPresence::AllowSparse);
+			FStructId Struct = Batch.DeclareStruct("Test", "Struct", {"MaybeB", "Bs", "MaybeBs", "B"}, EMemberPresence::AllowSparse);
 						
 			const bool True = true;
 			const bool False = false;
@@ -1315,14 +1302,14 @@ TEST_CASE_NAMED(FPlainPropsReadWriteTest, "System::Core::Serialization::PlainPro
 	{
 		TestSerialize([](FTestBatchBuilder& Batch, FScratchAllocator& Scratch)
 		{
-			FStructSchemaId Unused1 = Batch.DeclareStruct("Test", "Unused1", {"X"}, EMemberPresence::AllowSparse);
-			FStructSchemaId SA = Batch.DeclareStruct("Test", "SA", {"X"}, EMemberPresence::AllowSparse);
-			FStructSchemaId Unused2 = Batch.DeclareStruct("Test", "Unused2", {"X"}, EMemberPresence::AllowSparse);
-			FStructSchemaId SB = Batch.DeclareStruct("Test", "SB", {"X"}, EMemberPresence::AllowSparse);
-			FStructSchemaId Object = Batch.DeclareStruct("Test", "Object", {"Same", "Some", "None", "Diff"}, EMemberPresence::AllowSparse);
-			FStructSchemaId Unused3 = Batch.DeclareStruct("Test", "Unused3", {"X"}, EMemberPresence::AllowSparse);
+			FStructId Unused1 = Batch.DeclareStruct("Test", "Unused1", {"X"}, EMemberPresence::AllowSparse);
+			FStructId SA = Batch.DeclareStruct("Test", "SA", {"X"}, EMemberPresence::AllowSparse);
+			FStructId Unused2 = Batch.DeclareStruct("Test", "Unused2", {"X"}, EMemberPresence::AllowSparse);
+			FStructId SB = Batch.DeclareStruct("Test", "SB", {"X"}, EMemberPresence::AllowSparse);
+			FStructId Object = Batch.DeclareStruct("Test", "Object", {"Same", "Some", "None", "Diff"}, EMemberPresence::AllowSparse);
+			FStructId Unused3 = Batch.DeclareStruct("Test", "Unused3", {"X"}, EMemberPresence::AllowSparse);
 		
-			auto BuildStruct = [&](FStructSchemaId Struct, auto X)
+			auto BuildStruct = [&](FStructId Struct, auto X)
 				{
 					FMemberBuilder Members;
 					Members.Add(Batch.NameMember("X"), X);
@@ -1364,12 +1351,12 @@ TEST_CASE_NAMED(FPlainPropsReadWriteTest, "System::Core::Serialization::PlainPro
 	{
 		TestSerialize([](FTestBatchBuilder& Batch, FScratchAllocator& Scratch)
 		{
-			FStructSchemaId SA = Batch.DeclareStruct("Test", "SA", {"X"}, EMemberPresence::AllowSparse);
-			FStructSchemaId Unused = Batch.DeclareStruct("Test", "Unused2", {"X"}, EMemberPresence::AllowSparse);
-			FStructSchemaId SB = Batch.DeclareStruct("Test", "SB", {"X"}, EMemberPresence::AllowSparse);
-			FStructSchemaId Object = Batch.DeclareStruct("Test", "Object", {"Same", "Some", "None", "Diff", "SameEmpty", "DiffEmpty", "DiffNested"}, EMemberPresence::AllowSparse);
+			FStructId SA = Batch.DeclareStruct("Test", "SA", {"X"}, EMemberPresence::AllowSparse);
+			FStructId Unused = Batch.DeclareStruct("Test", "Unused2", {"X"}, EMemberPresence::AllowSparse);
+			FStructId SB = Batch.DeclareStruct("Test", "SB", {"X"}, EMemberPresence::AllowSparse);
+			FStructId Object = Batch.DeclareStruct("Test", "Object", {"Same", "Some", "None", "Diff", "SameEmpty", "DiffEmpty", "DiffNested"}, EMemberPresence::AllowSparse);
 		
-			auto BuildStructRange = [&](FStructSchemaId Struct, auto X)
+			auto BuildStructRange = [&](FStructId Struct, auto X)
 				{
 					FStructRangeBuilder Members(1);
 					Members[0].Add(Batch.NameMember("X"), X);
@@ -1426,10 +1413,10 @@ TEST_CASE_NAMED(FPlainPropsReadWriteTest, "System::Core::Serialization::PlainPro
 	{
 		TestSerialize([](FTestBatchBuilder& Batch, FScratchAllocator& Scratch)
 		{
-			FStructSchemaId Unused = Batch.DeclareStruct("Test", "X", {"X"}, EMemberPresence::AllowSparse);
-			FStructSchemaId Low = Batch.DeclareStruct("Test", "Low", {"LInt"}, EMemberPresence::AllowSparse);
-			FStructSchemaId Mid = Batch.DeclareStruct("Test", "Mid", {"MInt", "MLow"}, EMemberPresence::AllowSparse, ToOptional(Low));
-			FStructSchemaId Top = Batch.DeclareStruct("Test", "Top", {"TInt", "TLow", "TMids"}, EMemberPresence::AllowSparse, ToOptional(Mid));
+			FStructId Unused = Batch.DeclareStruct("Test", "X", {"X"}, EMemberPresence::AllowSparse);
+			FStructId Low = Batch.DeclareStruct("Test", "Low", {"LInt"}, EMemberPresence::AllowSparse);
+			FStructId Mid = Batch.DeclareStruct("Test", "Mid", {"MInt", "MLow"}, EMemberPresence::AllowSparse, ToOptional(Low));
+			FStructId Top = Batch.DeclareStruct("Test", "Top", {"TInt", "TLow", "TMids"}, EMemberPresence::AllowSparse, ToOptional(Mid));
 			
 			FMemberBuilder Members;
 			Members.Add(Batch.NameMember("LInt"), 123);
@@ -1507,14 +1494,14 @@ TEST_CASE_NAMED(FPlainPropsReadWriteTest, "System::Core::Serialization::PlainPro
 	{
 		TestSerialize([](FTestBatchBuilder& Batch, FScratchAllocator& Scratch)
 		{																										// Usage by A B C 
-			FStructSchemaId B0 = Batch.DeclareStruct("Test", "B0", {"0"}, EMemberPresence::AllowSparse);				 // - - -
-			FStructSchemaId B1 = Batch.DeclareStruct("Test", "B1", {"1"}, EMemberPresence::AllowSparse, ToOptional(B0)); // 1 1 1
-			FStructSchemaId B2 = Batch.DeclareStruct("Test", "B2", {"2"}, EMemberPresence::AllowSparse, ToOptional(B1)); // - 1 1
-			FStructSchemaId B3 = Batch.DeclareStruct("Test", "B3", {"3"}, EMemberPresence::AllowSparse, ToOptional(B2)); // - - 0
-			FStructSchemaId B4 = Batch.DeclareStruct("Test", "B4", {"4"}, EMemberPresence::AllowSparse, ToOptional(B3)); // 1 1 1
-			FStructSchemaId B5 = Batch.DeclareStruct("Test", "B5", {"5"}, EMemberPresence::AllowSparse, ToOptional(B4)); // 1 1 0
-			FStructSchemaId B6 = Batch.DeclareStruct("Test", "B6", {   }, EMemberPresence::AllowSparse, ToOptional(B5)); // 0 - 0
-			FStructSchemaId C5 = Batch.DeclareStruct("Test", "C5", {"5"}, EMemberPresence::AllowSparse, ToOptional(B4)); // - - -
+			FStructId B0 = Batch.DeclareStruct("Test", "B0", {"0"}, EMemberPresence::AllowSparse);				 // - - -
+			FStructId B1 = Batch.DeclareStruct("Test", "B1", {"1"}, EMemberPresence::AllowSparse, ToOptional(B0)); // 1 1 1
+			FStructId B2 = Batch.DeclareStruct("Test", "B2", {"2"}, EMemberPresence::AllowSparse, ToOptional(B1)); // - 1 1
+			FStructId B3 = Batch.DeclareStruct("Test", "B3", {"3"}, EMemberPresence::AllowSparse, ToOptional(B2)); // - - 0
+			FStructId B4 = Batch.DeclareStruct("Test", "B4", {"4"}, EMemberPresence::AllowSparse, ToOptional(B3)); // 1 1 1
+			FStructId B5 = Batch.DeclareStruct("Test", "B5", {"5"}, EMemberPresence::AllowSparse, ToOptional(B4)); // 1 1 0
+			FStructId B6 = Batch.DeclareStruct("Test", "B6", {   }, EMemberPresence::AllowSparse, ToOptional(B5)); // 0 - 0
+			FStructId C5 = Batch.DeclareStruct("Test", "C5", {"5"}, EMemberPresence::AllowSparse, ToOptional(B4)); // - - -
 			
 			FMemberBuilder A;
 			A.Add(Batch.NameMember("1"), 1);
@@ -1614,9 +1601,9 @@ TEST_CASE_NAMED(FPlainPropsReadWriteTest, "System::Core::Serialization::PlainPro
 			FType E2T{NestedUsed,		Batch.MakeTypename("E2")};
 			FType E3T{NestedUnused2,	Batch.MakeTypename("E3")};
 			
-			FEnumSchemaId E1D = Batch.DeclareEnum(E1T, EEnumMode::Flat, ELeafWidth::B8, {"C1"}, {1});
-			FEnumSchemaId E2D = Batch.DeclareEnum(E2T, EEnumMode::Flat, ELeafWidth::B8, {"C2"}, {2});
-			FEnumSchemaId E3D = Batch.DeclareEnum(E3T, EEnumMode::Flat, ELeafWidth::B8, {"C3"}, {3});
+			FEnumId E1D = Batch.DeclareEnum(E1T, EEnumMode::Flat, ELeafWidth::B8, {"C1"}, {1});
+			FEnumId E2D = Batch.DeclareEnum(E2T, EEnumMode::Flat, ELeafWidth::B8, {"C2"}, {2});
+			FEnumId E3D = Batch.DeclareEnum(E3T, EEnumMode::Flat, ELeafWidth::B8, {"C3"}, {3});
 
 			FType S1T{NestedUnused1,	Batch.MakeTypename("S1")};
 			FType S2T{NestedUsed,		Batch.MakeTypename("S2")};
@@ -1624,11 +1611,11 @@ TEST_CASE_NAMED(FPlainPropsReadWriteTest, "System::Core::Serialization::PlainPro
 			FType S4T = Batch.MakeParametricType({DoubleNested,	Batch.MakeTypename("S4")}, {S2T, E2T});
 			FType S5T = Batch.MakeParametricType({NestedUnused3,	Batch.MakeTypename("S5")}, {E3T, E1T, S2T});
 			
-			FStructSchemaId S1D = Batch.DeclareStruct(S1T, {"M1"}, EMemberPresence::AllowSparse);
-			FStructSchemaId S2D = Batch.DeclareStruct(S2T, {"M2"}, EMemberPresence::AllowSparse);
-			FStructSchemaId S3D = Batch.DeclareStruct(S3T, {"M3"}, EMemberPresence::AllowSparse);
-			FStructSchemaId S4D = Batch.DeclareStruct(S4T, {"M4"}, EMemberPresence::AllowSparse);
-			FStructSchemaId S5D = Batch.DeclareStruct(S5T, {"M5"}, EMemberPresence::AllowSparse);
+			FStructId S1D = Batch.DeclareStruct(S1T, {"M1"}, EMemberPresence::AllowSparse);
+			FStructId S2D = Batch.DeclareStruct(S2T, {"M2"}, EMemberPresence::AllowSparse);
+			FStructId S3D = Batch.DeclareStruct(S3T, {"M3"}, EMemberPresence::AllowSparse);
+			FStructId S4D = Batch.DeclareStruct(S4T, {"M4"}, EMemberPresence::AllowSparse);
+			FStructId S5D = Batch.DeclareStruct(S5T, {"M5"}, EMemberPresence::AllowSparse);
 
 			FMemberBuilder Members;
 			Members.Add(Batch.NameMember("M4"), 1);
diff --git a/Engine/Plugins/Experimental/PlainProps/Source/Tests/PlainPropsUeCoreTests.cpp b/Engine/Plugins/Experimental/PlainProps/Source/Tests/PlainPropsUeCoreTests.cpp
--- a/Engine/Plugins/Experimental/PlainProps/Source/Tests/PlainPropsUeCoreTests.cpp
+++ b/Engine/Plugins/Experimental/PlainProps/Source/Tests/PlainPropsUeCoreTests.cpp
@@ -43,8 +43,8 @@ struct FRuntimeIds
 	static FMemberId			IndexMember(FAnsiStringView Name)	{ return GNames.NameMember(FName(Name)); }
 	static FConcreteTypenameId	IndexTypename(FAnsiStringView Name)	{ return GNames.NameType(FName(Name)); }
 	static FScopeId				IndexScope(FAnsiStringView Name)	{ return GNames.MakeScope(FName(Name)); }
-	static FEnumSchemaId		IndexEnum(FType Type)				{ return GNames.IndexEnum(Type); }
-	static FStructSchemaId		IndexStruct(FType Type)			{ return GNames.IndexStruct(Type); }
+	static FEnumId				IndexEnum(FType Type)				{ return GNames.IndexEnum(Type); }
+	static FStructId			IndexStruct(FType Type)				{ return GNames.IndexStruct(Type); }
 	static FIdIndexerBase&		GetIndexer()						{ return GNames; }
 };
 
@@ -77,7 +77,7 @@ struct TScopedEnumDeclaration
 	using Ids = typename Runtime::Ids;
 	using Ctti = CttiOf<Enum>;
 
-	FEnumSchemaId Id;
+	FEnumId Id;
 	TScopedEnumDeclaration() : Id(DeclareNativeEnum<Ctti, Ids>(Runtime::GetTypes(), Mode)) {}
 	~TScopedEnumDeclaration() { Runtime::GetTypes().DropEnum(Id); }
 };
@@ -90,7 +90,7 @@ using TScopedDefaultStructBinding = TScopedStructBinding<T, FDefaultRuntime>;
 // Maybe replace with macro, e.g. PP_DECLARE_CUSTOM_DENSE_STRUCT(FRuntimeIds, FName, void, Idx)
 struct FNameDeclaration 
 {
-	FStructSchemaId		Id;
+	FStructId			Id;
 	FMemberId			Idx;
 
 	FNameDeclaration(FType Type = IndexStructName<FRuntimeIds, ETypename::Decl, TTypename<FName>>())
@@ -108,7 +108,7 @@ struct FNameDeclaration
 
 struct FTestCustomBinding : ICustomBinding
 {
-//	virtual FDualStructSchemaId GetId() const = 0;
+//	virtual FDualStructId GetId() const = 0;
 };
 
 struct FNameBinding : FTestCustomBinding
@@ -132,7 +132,7 @@ struct FNameBinding : FTestCustomBinding
 		return A.IsEqual(B, ENameCase::CaseSensitive);
 	}
 
-//	virtual FDualStructSchemaId GetId() const override { return { Declaration.Id, Declaration.Id }; }
+//	virtual FDualStructId GetId() const override { return { Declaration.Id, Declaration.Id }; }
 
 	FNameDeclaration	Declaration;
 	TSet<FName>			Names;
@@ -150,10 +150,10 @@ public:
 	template<class T>
 	bool						SaveDelta(const T& Object, const T& Default);
 
-	TArray64<uint8>				Write(TArray<FStructSchemaId>* OutMemoryIds) const;
+	TArray64<uint8>				Write(TArray<FStructId>* OutMemoryIds) const;
 
 private:
-	using IdBuiltStructPair = TPair<FStructSchemaId, const FBuiltStruct*>;
+	using IdBuiltStructPair = TPair<FStructId, const FBuiltStruct*>;
 	TArray<IdBuiltStructPair>	SavedObjects;
 	FNameBinding				SavedNames;
 	FCustomBindingsOverlay		Customs;
@@ -169,14 +169,14 @@ FBatchSaver::FBatchSaver(const FCustomBindings& Underlay)
 template<class T>
 void FBatchSaver::Save(T&& Object) 
 {
-	FDualStructSchemaId Id = IndexStructDualId<FRuntimeIds, TTypename<std::remove_reference_t<T>>>();
+	FDualStructId Id = IndexStructDualId<FRuntimeIds, TTypename<std::remove_reference_t<T>>>();
 	SavedObjects.Emplace(Id.DeclId, SaveStruct(&Object, Id.BindId, {GTypes, GSchemas, Customs, Scratch}));
 }
 
 template<class T>
 bool FBatchSaver::SaveDelta(const T& Object, const T& Default) 
 {
-	FDualStructSchemaId Id = IndexStructDualId<FRuntimeIds, TTypename<std::remove_reference_t<T>>>();
+	FDualStructId Id = IndexStructDualId<FRuntimeIds, TTypename<std::remove_reference_t<T>>>();
 	if (const FBuiltStruct* Delta = SaveStructDeltaIfDiff(&Object, &Default, Id.BindId, {GTypes, GSchemas, Customs, Scratch}))
 	{
 		SavedObjects.Emplace(Id.DeclId, MoveTemp(Delta));
@@ -201,7 +201,7 @@ TConstArrayView<T> GrabNumAndArray(/* in-out */ FByteReader& It)
 
 inline constexpr uint32 Magics[] = { 0xFEEDF00D, 0xABCD1234, 0xDADADAAA, 0x99887766, 0xF0F1F2F3 , 0x00112233};
 
-TArray64<uint8> FBatchSaver::Write(TArray<FStructSchemaId>* OutMemoryIds) const
+TArray64<uint8> FBatchSaver::Write(TArray<FStructId>* OutMemoryIds) const
 {
 	ESchemaFormat Format = OutMemoryIds ? ESchemaFormat::InMemoryNames : ESchemaFormat::StableNames;
 
@@ -242,7 +242,7 @@ TArray64<uint8> FBatchSaver::Write(TArray<FStructSchemaId>* OutMemoryIds) const
 
 	// Write objects
 	WriteInt(Out, Magics[2]);
-	for (const TPair<FStructSchemaId, const FBuiltStruct*>& Object : SavedObjects)
+	for (const TPair<FStructId, const FBuiltStruct*>& Object : SavedObjects)
 	{
 		WriteInt(/* out */ Tmp, Magics[3]);
 		WriteInt(/* out */ Tmp, Writer.GetWriteId(Object.Key).Get().Idx);
@@ -264,112 +264,88 @@ TArray64<uint8> FBatchSaver::Write(TArray<FStructSchemaId>* OutMemoryIds) const
 
 ////////////////////////////////////////////////////////////////////////////////////////////////
 
-class FSchemaBatchIds : public FIds
-{
-public:
-	FSchemaBatchIds(const FSchemaBatch& InSchemas) : Schemas(InSchemas) {}
-
-	virtual uint32				NumEnums() const override final						{ return Schemas.NumSchemas - Schemas.NumStructSchemas; }
-	virtual uint32				NumStructs() const override final					{ return Schemas.NumStructSchemas; }
-
-	virtual FType				Resolve(FEnumSchemaId Id) const override final		{ return ResolveEnumSchema(Schemas, Id).Type; }
-	virtual FType				Resolve(FStructSchemaId Id) const override final	{ return ResolveStructSchema(Schemas, Id).Type; }
-
-protected:
-	const FSchemaBatch& Schemas;
-};
-
-class FMemoryBatchIds final : public FSchemaBatchIds
-{
-public:
-	using FSchemaBatchIds::FSchemaBatchIds;
-	using FSchemaBatchIds::Resolve;
-
-	virtual uint32				NumNames() const override final						{ return GNames.NumNames(); }
-	virtual uint32				NumNestedScopes() const override final				{ return GNames.NumNestedScopes(); }
-	virtual uint32				NumParametricTypes() const override final			{ return GNames.NumParametricTypes(); }
-
-	virtual FNestedScope		Resolve(FNestedScopeId Id) const override			{ return GNames.Resolve(Id); }
-	virtual FParametricTypeView	Resolve(FParametricTypeId Id) const override		{ return GNames.Resolve(Id); }
-	virtual void				AppendString(FUtf8Builder& Out, FNameId Name) const override final { GNames.ResolveName(Name).AppendString(Out); }
-};
+//class FSchemaBatchIds : public FIds
+//{
+//public:
+//	FSchemaBatchIds(const FSchemaBatch& InSchemas) : Schemas(InSchemas) {}
+//
+//	virtual uint32				NumEnums() const override final						{ return Schemas.NumSchemas - Schemas.NumStructSchemas; }
+//	virtual uint32				NumStructs() const override final					{ return Schemas.NumStructSchemas; }
+//
+//	virtual FType				Resolve(FEnumSchemaId Id) const override final		{ return ResolveEnumSchema(Schemas, Id).Type; }
+//	virtual FType				Resolve(FStructSchemaId Id) const override final	{ return ResolveStructSchema(Schemas, Id).Type; }
+//
+//protected:
+//	const FSchemaBatch& Schemas;
+//};
+//
+//class FMemoryBatchIds final : public FSchemaBatchIds
+//{
+//public:
+//	using FSchemaBatchIds::FSchemaBatchIds;
+//	using FSchemaBatchIds::Resolve;
+//
+//	virtual uint32				NumNames() const override final						{ return GNames.NumNames(); }
+//	virtual uint32				NumNestedScopes() const override final				{ return GNames.NumNestedScopes(); }
+//	virtual uint32				NumParametricTypes() const override final			{ return GNames.NumParametricTypes(); }
+//
+//	virtual FNestedScope		Resolve(FNestedScopeId Id) const override			{ return GNames.Resolve(Id); }
+//	virtual FParametricTypeView	Resolve(FParametricTypeId Id) const override		{ return GNames.Resolve(Id); }
+//	virtual void				AppendString(FUtf8Builder& Out, FNameId Name) const override { GNames.ResolveName(Name).AppendString(Out); }
+//};
 
-class FStableBatchIds final : public FSchemaBatchIds
+class FStableNameBatchIds final : public FStableBatchIds
 {
+	TConstArrayView<FName> Names;
 public:
-	using FSchemaBatchIds::FSchemaBatchIds;
-	using FSchemaBatchIds::Resolve;
-
-	virtual uint32				NumNames() const override final						{ return IntCastChecked<uint32>(Names.Num()); }
-	virtual uint32				NumNestedScopes() const override final				{ return Schemas.NumNestedScopes; }
-	virtual uint32				NumParametricTypes() const override final			{ return Schemas.NumParametricTypes; }
+	FStableNameBatchIds(FReadBatchId Batch, TConstArrayView<FName> InNames) : FStableBatchIds(Batch), Names(InNames) {}
+	using FStableBatchIds::AppendString;
 
-	virtual FNestedScope		Resolve(FNestedScopeId Id) const override			{ return ResolveNestedScope(Schemas, Id); }
-	virtual FParametricTypeView	Resolve(FParametricTypeId Id) const override		{ return ResolveParametricType(Schemas, Id); }
-	virtual void				AppendString(FUtf8Builder& Out, FNameId Name) const override final { Names[Name.Idx].AppendString(Out); }
-
-	TConstArrayView<FName> Names;
+	virtual uint32				NumNames() const override							{ return static_cast<uint32>(Names.Num()); }
+	virtual void				AppendString(FUtf8Builder& Out, FNameId Name) const override { Names[Name.Idx].AppendString(Out); }
 };
 
-class FLoadBatchNames : public FIds
+class FLoadBatchNames final : public FBatchIds
 {
+	const FIdBinding& Binding;
 public:
-	FLoadBatchNames(const FIdBinding& InBinding, const FSchemaBatch& InSchemas) : Ids(GNames), Binding(InBinding), Schemas(InSchemas) {}
+	using FBatchIds::Resolve;
+	using FBatchIds::AppendString;
 
-	virtual uint32				NumNames() const override final						{ return Ids.NumNames(); }
-	virtual uint32				NumNestedScopes() const override final				{ return IntCastChecked<uint32>(Binding.NestedScopes.Num()); }
-	virtual uint32				NumParametricTypes() const override final			{ return IntCastChecked<uint32>(Binding.ParametricTypes.Num()); }
-	virtual uint32				NumEnums() const override final						{ return Schemas.NumSchemas - Schemas.NumStructSchemas; }
-	virtual uint32				NumStructs() const override final					{ return Schemas.NumStructSchemas; }
+	FLoadBatchNames(const FIdBinding& InBinding, FReadBatchId Batch) : FBatchIds(Batch), Binding(InBinding) {}
 
-	virtual FNestedScope		Resolve(FNestedScopeId Id) const override final		{ return Ids.Resolve(Binding.Remap(Id)); }
-	virtual FParametricTypeView	Resolve(FParametricTypeId Id) const override final	{ return Ids.Resolve(Binding.Remap(Id)); }
-	virtual FType				Resolve(FEnumSchemaId Id) const override final		{ return ResolveEnumSchema(Schemas, Id).Type; }
-	virtual FType				Resolve(FStructSchemaId Id) const override final	{ return ResolveStructSchema(Schemas, Id).Type; }
+	virtual uint32				NumNames() const override												{ return GNames.NumNames(); }
+	virtual void				AppendString(FUtf8Builder& Out, FNameId Name) const override			{ PlainProps::AppendString(Out, GNames.ResolveName(Name)); }
+	virtual void				AppendString(FUtf8Builder& Out, FTypenameId Typename) const override	{ GNames.AppendString(Out, Typename); }
+	virtual void				AppendString(FUtf8Builder& Out, FScopeId Scope) const override			{ GNames.AppendString(Out, Scope); }
 
-	using FIds::AppendString;
-	virtual void AppendString(FUtf8Builder& Out, FNameId Name) const override final			{ Ids.AppendString(Out, Name); }
-	virtual void AppendString(FUtf8Builder& Out, FTypenameId Typename) const override final	{ Ids.AppendString(Out, Typename); }
-	virtual void AppendString(FUtf8Builder& Out, FScopeId Scope) const override final		{ Ids.AppendString(Out, Scope); }
-
-private:
-	const FIds& Ids;
-	const FIdBinding& Binding;
-	const FSchemaBatch& Schemas;
+	virtual uint32				NumNestedScopes() const override					{ return static_cast<uint32>(Binding.NestedScopes.Num()); }
+	virtual uint32				NumParametricTypes() const override					{ return static_cast<uint32>(Binding.ParametricTypes.Num()); }
+	virtual FNestedScope		Resolve(FNestedScopeId Id) const override			{ return GNames.Resolve(Binding.Remap(Id)); }
+	virtual FParametricTypeView	Resolve(FParametricTypeId Id) const override		{ return GNames.Resolve(Binding.Remap(Id)); }
 };
 
-static void PrintBatch(
-	FUtf8Builder& Out,
-	const FSchemaBatchIds& Ids,
-	const FSchemaBatch& Schemas,
-	TArray<FStructView>&& Objects)
+static void PrintBatch(FUtf8Builder& Out, TConstArrayView<FName> Names,	const FSchemaBatch& Schemas, TArray<FStructView>&& Objects)
 {
 	FReadBatchId Batch = MountReadSchemas(&Schemas);
+
 	for (FStructView& Struct : Objects)
 	{
 		Struct.Schema.Batch = Batch;
 	}
-	PrintYamlBatch(Out, Ids, Batch, Objects);
-	UnmountReadSchemas(Batch);
-}
-
-static void PrintBatch(
-	FUtf8Builder& Out,
-	TConstArrayView<FName> Names,
-	const FSchemaBatch& Schemas,
-	TConstArrayView<FStructView> Objects)
-{
+	
 	if (Names.IsEmpty())
 	{
-		FMemoryBatchIds Ids(Schemas);
-		PrintBatch(Out, Ids, Schemas, TArray<FStructView>(Objects));
+		FMemoryBatchIds Ids(Batch, GNames);
+		PrintYamlBatch(Out, Ids, Objects);
 	}
 	else
 	{
-		FStableBatchIds Ids(Schemas);
-		Ids.Names = Names;
-		PrintBatch(Out, Ids, Schemas, TArray<FStructView>(Objects));
+		FStableNameBatchIds Ids(Batch, Names);
+		PrintYamlBatch(Out, Ids, Objects);
 	}
+	
+	UnmountReadSchemas(Batch);
 
 	FUtf8StringView YamlView = Out.ToView();
 	auto Wide = StringCast<TCHAR>(YamlView.GetData(), YamlView.Len());
@@ -398,8 +374,8 @@ static void PrintLoadStruct(const FIdBinding& Binding, FStructView StructView, i
 	TUtf8StringBuilder<4096> YamlString;
 	{
 		FYamlBuilderPtr YamlBuilder = MakeYamlBuilder(YamlString);
-		FLoadBatchNames BatchNames(Binding, GetReadSchemas(StructView.Schema.Batch));
-		FBatchPrinter Printer(*YamlBuilder, BatchNames);
+		FLoadBatchNames Ids(Binding, StructView.Schema.Batch);
+		FBatchPrinter Printer(*YamlBuilder, Ids);
 		Printer.PrintObjects({StructView});
 	}
 	FUtf8StringView YamlView = YamlString.ToView();
@@ -412,7 +388,7 @@ static void PrintLoadStruct(const FIdBinding& Binding, FStructView StructView, i
 class FBatchLoader
 {
 public:
-	FBatchLoader(FMemoryView Data, const FCustomBindings& Underlay, TConstArrayView<FStructSchemaId> InRuntimeIds)
+	FBatchLoader(FMemoryView Data, const FCustomBindings& Underlay, TConstArrayView<FStructId> InRuntimeIds)
 	: Customs(Underlay)
 	{
 		ESchemaFormat Format = InRuntimeIds.IsEmpty() ? ESchemaFormat::StableNames : ESchemaFormat::InMemoryNames;
@@ -449,8 +425,8 @@ public:
 			// Mount saved schemas as is
 			Batch = MountReadSchemas(SavedSchemas);	
 
-			TArray<FStructSchemaId> lala = IndexRuntimeIds(*SavedSchemas, GNames);
-			check(EqualItems(lala, InRuntimeIds));
+			TArray<FStructId> ExpectedRuntimeIds = IndexRuntimeIds(*SavedSchemas, GNames);
+			check(EqualItems(ExpectedRuntimeIds, InRuntimeIds));
 		}
 		
 		// Read objects
@@ -471,13 +447,13 @@ public:
 		CHECK(It.Grab<uint32>() == Magics[5]);
 
 		// Finally create load plans
-		TConstArrayView<FStructSchemaId> LoadStructIds = RuntimeIds ? RuntimeIds->Translation.GetStructIds(SavedSchemas->NumStructSchemas) : InRuntimeIds;
+		TConstArrayView<FStructId> LoadStructIds = RuntimeIds ? RuntimeIds->Translation.GetStructIds(SavedSchemas->NumStructSchemas) : InRuntimeIds;
 		Plans = CreateLoadPlans(Batch, GTypes, Customs, GSchemas, LoadStructIds, Format);
 
 		// Mount, print, unmount, and then parse the saved batch
-		TUtf8StringBuilder<4096> YamlString;
-		PrintBatch(YamlString, Ids, *SavedSchemas, Objects);
-		ParseBatch(YamlString.ToView(), SavedSchemasView, Format);
+		TUtf8StringBuilder<4096> YamlBatch;
+		PrintBatch(YamlBatch, Ids, *SavedSchemas, TArray<FStructView>(Objects));
+		ParseBatch(YamlBatch.ToView(), SavedSchemasView, Format);
 	}
 
 	~FBatchLoader()
@@ -535,7 +511,7 @@ static void Run(void (*Save)(FBatchSaver&), void (*Load)(FBatchLoader&), const F
 	for (ESchemaFormat Format : {ESchemaFormat::StableNames, ESchemaFormat::InMemoryNames})
 	{
 		TArray64<uint8> Data;
-		TArray<FStructSchemaId>	RuntimeIds;
+		TArray<FStructId>	RuntimeIds;
 		{
 			FBatchSaver Batch(Customs);
 			Save(Batch);
diff --git a/Engine/Plugins/Experimental/PlainPropsUObject/Source/Private/PlainPropsUObjectBindings.cpp b/Engine/Plugins/Experimental/PlainPropsUObject/Source/Private/PlainPropsUObjectBindings.cpp
--- a/Engine/Plugins/Experimental/PlainPropsUObject/Source/Private/PlainPropsUObjectBindings.cpp
+++ b/Engine/Plugins/Experimental/PlainPropsUObject/Source/Private/PlainPropsUObjectBindings.cpp
@@ -135,7 +135,7 @@ static void ReserveZeroes(/* in-out */ FMutableMemoryView& Zeroes, SIZE_T Size,
 
 FDefaultStructs::~FDefaultStructs()
 {
-	for (TPair<FStructSchemaId, FDefaultInstance> Instance : Instances)
+	for (TPair<FStructId, FDefaultInstance> Instance : Instances)
 	{
 		DeleteInstance(Instance.Value);
 	}
@@ -158,7 +158,7 @@ void FDefaultStructs::ReserveFlags(uint32 Idx)
 	}
 }
 
-void FDefaultStructs::Bind(FStructSchemaId Id, const UScriptStruct* Struct)
+void FDefaultStructs::Bind(FStructId Id, const UScriptStruct* Struct)
 {
 	const EStructFlags Flags = Struct->StructFlags;
 	const SIZE_T Size = Struct->GetStructureSize();
@@ -205,7 +205,7 @@ void FDefaultStructs::Bind(FStructSchemaId Id, const UScriptStruct* Struct)
 	}
 }
 
-void FDefaultStructs::BindZeroes(FStructSchemaId Id, SIZE_T Size, uint32 Alignment)
+void FDefaultStructs::BindZeroes(FStructId Id, SIZE_T Size, uint32 Alignment)
 {
 	ReserveFlags(Id.Idx);
 #if DO_CHECK
@@ -214,7 +214,7 @@ void FDefaultStructs::BindZeroes(FStructSchemaId Id, SIZE_T Size, uint32 Alignme
 	ReserveZeroes(Zeroes, Size, Alignment);
 }
 
-void FDefaultStructs::BindStatic(FStructSchemaId Id, const void* Struct)
+void FDefaultStructs::BindStatic(FStructId Id, const void* Struct)
 {
 	ReserveFlags(Id.Idx);
 #if DO_CHECK
@@ -227,7 +227,7 @@ void FDefaultStructs::BindStatic(FStructSchemaId Id, const void* Struct)
 	Instances.Add(Id, MakeStaticInstance(Struct));
 }
 
-void FDefaultStructs::Drop(FStructSchemaId Id)
+void FDefaultStructs::Drop(FStructId Id)
 {
 #if DO_CHECK
 	checkf(!Flip(Bound[Id.Idx]), TEXT("'%s' isn't bound"), *GUE.Debug.Print(Id));
@@ -239,7 +239,7 @@ void FDefaultStructs::Drop(FStructSchemaId Id)
 	}
 }
 
-const void* FDefaultStructs::Get(FStructSchemaId Id)
+const void* FDefaultStructs::Get(FStructId Id)
 {
 #if DO_CHECK
 	checkf(Bound[Id.Idx], TEXT("'%s' lack default"), *GUE.Debug.Print(Id));
@@ -352,7 +352,7 @@ static constexpr uint64 CommonStructMask =	CASTCLASS_FNameProperty | CASTCLASS_F
 											CASTCLASS_FDelegateProperty | CASTCLASS_FMulticastInlineDelegateProperty;
 static constexpr uint64 MiscMask =			CASTCLASS_FMulticastSparseDelegateProperty | CASTCLASS_FInterfaceProperty;
 
-static FStructSchemaId FlagsToCommonBindId(uint64 MaskedCastFlags)
+static FStructId FlagsToCommonBindId(uint64 MaskedCastFlags)
 {
 	switch (MaskedCastFlags)
 	{
@@ -449,7 +449,7 @@ static FType IndexType(const UField* Field)
 	return { GUE.Names.NestReversedScopes(ReversedOuters), Name };
 }
 
-static FOptionalStructSchemaId IndexSuper(const UStruct* Struct)
+static FOptionalStructId IndexSuper(const UStruct* Struct)
 {
 	if (const UStruct* Super = Struct->GetInheritanceSuper())
 	{
@@ -499,9 +499,9 @@ static void DeclareSuperMembers(FMemberArray& Out, const UStruct* Struct)
 	}
 }
 
-static void DeclareStruct(const UStruct* Struct, FType Type, FStructSchemaId Id)
+static void DeclareStruct(const UStruct* Struct, FType Type, FStructId Id)
 {
-	FOptionalStructSchemaId  Super = IndexSuper(Struct);
+	FOptionalStructId  Super = IndexSuper(Struct);
 	EMemberPresence Occupancy = GetOccupancy(Struct);
 	
 	FMemberArray Members;
@@ -516,10 +516,10 @@ static void DeclareStruct(const UStruct* Struct, FType Type, FStructSchemaId Id)
 	GUE.Types.DeclareStruct(Id, Type, Members, Occupancy, Super);
 }
 
-static FStructSchemaId DeclareStruct(const UStruct* Struct)
+static FStructId DeclareStruct(const UStruct* Struct)
 {
 	FType Type = IndexType(Struct);
-	FStructSchemaId Id = GUE.Names.IndexStruct(Type);
+	FStructId Id = GUE.Names.IndexStruct(Type);
 	DeclareStruct(Struct, Type, Id);
 	return Id;
 }
@@ -528,7 +528,7 @@ static FStructSchemaId DeclareStruct(const UStruct* Struct)
 
 static FTypedRange SaveNames(TConstArrayView<FName> Names, const FSaveContext& Ctx)
 {
-	FStructSchemaId Id = GUE.Structs.Name;
+	FStructId Id = GUE.Structs.Name;
 	FStructRangeSaver Out(Ctx.Scratch, Names.Num());
 	for (FName Name : Names)
 	{
@@ -615,7 +615,7 @@ static void LoadUnicastDelegate(FUnicastScriptDelegate& Dst, FStructLoadView Src
 
 static FTypedRange SaveInvocations(const FMulticastInvocationList& In, const FSaveContext& Ctx)
 {
-	FStructSchemaId ItemId = GUE.Structs.Delegate;
+	FStructId ItemId = GUE.Structs.Delegate;
 	FMemberSchema Schema = MakeStructRangeSchema(DefaultRangeMax, ItemId);
 	if (int32 NumTotal = In.Num())
 	{
@@ -805,7 +805,7 @@ struct FMulticastSparseDelegateBinding final : ICustomBinding
 	}
 };
 
-static FStructSchemaId BindSparseDelegate(FStructSchemaId Owner, FMulticastSparseDelegateProperty* Property)
+static FStructId BindSparseDelegate(FStructId Owner, FMulticastSparseDelegateProperty* Property)
 {
 	// Todo: Ownership / memory leak
 	ICustomBinding* Leak = new FMulticastSparseDelegateBinding(CastChecked<USparseDelegateFunction>(Property->SignatureFunction));
@@ -814,7 +814,7 @@ static FStructSchemaId BindSparseDelegate(FStructSchemaId Owner, FMulticastSpars
 	FType OwnerParam = GUE.Names.Resolve(Owner);
 	FType PropertyParam = { GUE.Scopes.CoreUObject, FTypenameId(GUE.Names.NameType(Property->GetFName())) };
 	FType UniqueBindName = GUE.Names.MakeParametricType(MulticastSparseDelegate, {OwnerParam, PropertyParam});
-	FStructSchemaId BindId = GUE.Names.IndexStruct(UniqueBindName);
+	FStructId BindId = GUE.Names.IndexStruct(UniqueBindName);
 
 	GUE.Customs.BindStruct(BindId, *Leak, GUE.Types.Get(GUE.Structs.MulticastDelegate));
 
@@ -921,7 +921,7 @@ class FInterfaceBindings
 	const FType ScriptInterface;
 	const FStructDeclaration& Declaration;
 
-	TMap<FType, FStructSchemaId> BoundClasses;
+	TMap<FType, FStructId> BoundClasses;
 
 public:
 	FInterfaceBindings()
@@ -929,16 +929,16 @@ public:
 	, Declaration(GUE.Types.DeclareStruct(GUE.Structs.ScriptInterface, ScriptInterface, {GUE.Members.Object}, EMemberPresence::RequireAll))
 	{}
 
-	FStructSchemaId Bind(FInterfaceProperty* Property)
+	FStructId Bind(FInterfaceProperty* Property)
 	{
 		FType Class = IndexType(Property->InterfaceClass);
-		if (const FStructSchemaId* BindId = BoundClasses.Find(Class))
+		if (const FStructId* BindId = BoundClasses.Find(Class))
 		{
 			return *BindId;
 		}
 		
 		FType UniqueBindName = GUE.Names.MakeParametricType(ScriptInterface, {Class});
-		FStructSchemaId BindId = GUE.Names.IndexStruct(UniqueBindName);
+		FStructId BindId = GUE.Names.IndexStruct(UniqueBindName);
 		BoundClasses.Emplace(Class, BindId);
 
 		// Todo: Ownership / memory leak
@@ -949,7 +949,7 @@ public:
 	}
 };
 
-static FStructSchemaId BindInterface(FInterfaceProperty* Property)
+static FStructId BindInterface(FInterfaceProperty* Property)
 {
 	static FInterfaceBindings Bindings;
 	return Bindings.Bind(Property);
@@ -1012,10 +1012,10 @@ static bool DenyMaxValue(FName Enum)
 	return !Algo::Find(AllowsMax, Enum);
 }
 
-static FEnumSchemaId DeclareEnum(UEnum* Enum)
+static FEnumId DeclareEnum(UEnum* Enum)
 {
 	FType Type = IndexType(Enum);
-	FEnumSchemaId Id = GUE.Names.IndexEnum(Type);
+	FEnumId Id = GUE.Names.IndexEnum(Type);
 	EEnumMode Mode = Enum->HasAnyEnumFlags(EEnumFlags::Flags) ? EEnumMode::Flag : EEnumMode::Flat;
 
 	// Skip _MAX and _All enumerators
@@ -1581,7 +1581,7 @@ struct TLeafRangeSerializer
 	static constexpr SIZE_T	Size = sizeof(EquivalentType);
 
 	FMemberType				InnerType;
-	FOptionalSchemaId		EnumId;
+	FOptionalInnerId		EnumId;
 
 	explicit TLeafRangeSerializer(FMemberBinding In)
 	: InnerType(ToLeafType(In.InnermostType.AsLeaf()))
@@ -1622,18 +1622,18 @@ struct FStructRangeSerializer
 	using RangeSaver = FStructRangeSaver;
 	
 	FMemberType				InnerType;
-	FStructSchemaId			SaveId;
+	FStructId				SaveId;
 
 	explicit FStructRangeSerializer(FMemberBinding Item)
 	: InnerType(Item.InnermostType.AsStruct())
-	, SaveId(static_cast<FStructSchemaId>(Item.InnermostSchema.Get()))
+	, SaveId(Item.InnermostSchema.Get().AsStruct())
 	{
 		check(Item.RangeBindings.IsEmpty());
 	}
 
 	inline FMemberSchema MakeMemberSchema() const
 	{
-		return { FMemberType(DefaultRangeMax), InnerType, 1, FOptionalSchemaId(SaveId), nullptr };
+		return { FMemberType(DefaultRangeMax), InnerType, 1, FInnerId(SaveId), nullptr };
 	}
 
 	inline FBuiltStruct* SaveItem(const void* In, const FSaveContext& Ctx) const
@@ -1651,7 +1651,7 @@ struct FNestedRangeSerializer
 {
 	using RangeSaver = FNestedRangeSaver;
 
-	FOptionalSchemaId								InnermostSaveId;
+	FOptionalInnerId								InnermostSaveId;
 	uint16											NumInners;
 	TArray<FMemberType, TInlineAllocator<8>>		InnerTypes;
 	TArray<FMemberBindType, TInlineAllocator<8>>	InnerBindTypes;
@@ -2208,17 +2208,17 @@ struct TSetPropertyBinding : IItemRangeBinding, ICustomBinding
 
 class FSetBindings
 {
-	TMap<FParameterBinding, FStructSchemaId> Bindings;
+	TMap<FParameterBinding, FStructId> Bindings;
 
 	template<EPropertyKind ItemKind>
-	void BindNew(FSetProperty* Property, FMemberBinding Elem, FStructSchemaId BindId, const FStructDeclaration& Declaration) 
+	void BindNew(FSetProperty* Property, FMemberBinding Elem, FStructId BindId, const FStructDeclaration& Declaration) 
 	{
 		// Todo: Ownership / memory leak
 		TSetPropertyBinding<ItemKind>* Leak = new TSetPropertyBinding<ItemKind>(Property, Elem);
 		GUE.Customs.BindStruct(BindId, *Leak, Declaration, {});
 	}
 
-	void DeltaBindNew(FSetProperty* Property, FMemberBinding Elem, FStructSchemaId BindId, const FStructDeclaration& Declaration)
+	void DeltaBindNew(FSetProperty* Property, FMemberBinding Elem, FStructId BindId, const FStructDeclaration& Declaration)
 	{
 		switch (GetPropertyKind(Elem))
 		{
@@ -2236,10 +2236,10 @@ class FSetBindings
 	}
 public:
 
-	FStructSchemaId Bind(FSetProperty* Property, FMemberBinding Elem)
+	FStructId Bind(FSetProperty* Property, FMemberBinding Elem)
 	{
 		check(Elem.Offset == 0);
-		if (const FStructSchemaId* BindId = Bindings.Find(FParameterBinding(Elem)))
+		if (const FStructId* BindId = Bindings.Find(FParameterBinding(Elem)))
 		{
 			return *BindId;
 		}
@@ -2250,8 +2250,8 @@ public:
 		FType DeclType = Param.IsLowered()
 						 ? FType{ GUE.Scopes.Core, FTypenameId(GUE.Names.MakeParametricTypeId(GUE.Typenames.Set, MakeArrayView(&Param.DeclType, 1))) }
 						 : BindType;
-		FStructSchemaId BindId = GUE.Names.IndexStruct(BindType);
-		FStructSchemaId DeclId = Param.IsLowered() ? GUE.Names.IndexStruct(DeclType) : BindId;
+		FStructId BindId = GUE.Names.IndexStruct(BindType);
+		FStructId DeclId = Param.IsLowered() ? GUE.Names.IndexStruct(DeclType) : BindId;
 		FMemberId Members[] = { GUE.Members.Assign, GUE.Members.Remove, GUE.Members.Insert };
 
 		// Todo: Ownership / memory leak
@@ -2571,11 +2571,11 @@ struct TMapPropertyCustomBinding : TMapPropertyItemBinding<FScriptMap, KeyKind,
 
 class FMapBindings
 {
-	TMap<FStructSchemaId, FStructSchemaId>	NormalBindings;
-	TMap<FStructSchemaId, FRangeBinding>	FrozenBindings;
+	TMap<FStructId, FStructId>	NormalBindings;
+	TMap<FStructId, FRangeBinding>	FrozenBindings;
 
 	template<EPropertyKind KeyKind, EPropertyKind ValueKind>
-	static IItemRangeBinding* New3(FMapProperty* Property, FMapMemberBindings Members, FStructSchemaId* OutCustomId)
+	static IItemRangeBinding* New3(FMapProperty* Property, FMapMemberBindings Members, FStructId* OutCustomId)
 	{
 		if (OutCustomId == nullptr) // Freezable maps aren't delta serialized
 		{
@@ -2583,10 +2583,10 @@ class FMapBindings
 		}
 
 		// Index custom delta binding struct name
-		FParametricTypeId PairTypename = GUE.Names.Resolve(static_cast<FStructSchemaId>(Members.Pair.InnermostSchema.Get())).Name.AsParametric();
+		FParametricTypeId PairTypename = GUE.Names.Resolve(Members.Pair.InnermostSchema.Get().AsStruct()).Name.AsParametric();
 		TConstArrayView<FType> Params = GUE.Names.Resolve(PairTypename).GetParameters();
 		FType Type = { GUE.Scopes.Core, FTypenameId(GUE.Names.MakeParametricTypeId(GUE.Typenames.Map, Params)) };
-		FStructSchemaId Id = GUE.Names.IndexStruct(Type);
+		FStructId Id = GUE.Names.IndexStruct(Type);
 		FMemberId MemberIds[] = { GUE.Members.Assign, GUE.Members.Remove, GUE.Members.Insert };
 
 		// Todo: Ownership / memory leak
@@ -2599,7 +2599,7 @@ class FMapBindings
 	}
 
 	template<EPropertyKind KeyKind>
-	inline IItemRangeBinding* New2(FMapProperty* Property, FMapMemberBindings Members, FStructSchemaId* OutCustomId)
+	inline IItemRangeBinding* New2(FMapProperty* Property, FMapMemberBindings Members, FStructId* OutCustomId)
 	{
 		switch (GetPropertyKind(Members.Value))
 		{
@@ -2616,7 +2616,7 @@ class FMapBindings
 		}
 	}
 
-	inline IItemRangeBinding* New(FMapProperty* Property, FMapMemberBindings Members, FStructSchemaId* OutCustomId)
+	inline IItemRangeBinding* New(FMapProperty* Property, FMapMemberBindings Members, FStructId* OutCustomId)
 	{
 		switch (GetPropertyKind(Members.Key))
 		{
@@ -2634,19 +2634,19 @@ class FMapBindings
 	}
 
 public:
-	FStructSchemaId BindNormal(FMapProperty* Property, FStructSchemaId PairId, FMapMemberBindings Members)
+	FStructId BindNormal(FMapProperty* Property, FStructId PairId, FMapMemberBindings Members)
 	{
-		if (const FStructSchemaId* CustomId = NormalBindings.Find(PairId))
+		if (const FStructId* CustomId = NormalBindings.Find(PairId))
 		{
 			return *CustomId;
 		}
 
-		FStructSchemaId CustomId;
+		FStructId CustomId;
 		New(Property, Members, /* out */ &CustomId);
 		return NormalBindings.Emplace(PairId, CustomId);
 	}
 
-	FRangeBinding BindFreezable(FMapProperty* Property, FStructSchemaId PairId, FMapMemberBindings Members)
+	FRangeBinding BindFreezable(FMapProperty* Property, FStructId PairId, FMapMemberBindings Members)
 	{
 		if (const FRangeBinding* RangeBinding = FrozenBindings.Find(PairId))
 		{
@@ -2670,12 +2670,12 @@ class FPairBindings
 		inline bool operator==(const FPair& O) const { return EqSkipOffset(KV[0], O.KV[0]) && EqSkipOffset(KV[1], O.KV[1]); }
 	};
 
-	TMap<FPair, FStructSchemaId>	Bindings;
+	TMap<FPair, FStructId>	Bindings;
 
-	inline FStructSchemaId BindImpl(FPair Pair)
+	inline FStructId BindImpl(FPair Pair)
 	{
 		check(Pair.KV[0].Offset == 0 && Pair.KV[1].Offset > 0);
-		if (const FStructSchemaId* BindId = Bindings.Find(Pair))
+		if (const FStructId* BindId = Bindings.Find(Pair))
 		{
 			return *BindId;
 		}
@@ -2687,8 +2687,8 @@ class FPairBindings
 		FType DeclParams[2] = { KeyParam.DeclType, ValueParam.DeclType };
 		FType BindType = { GUE.Scopes.Core, FTypenameId(GUE.Names.MakeParametricTypeId(GUE.Typenames.Pair, BindParams)) };
 		FType DeclType = { GUE.Scopes.Core, FTypenameId(GUE.Names.MakeParametricTypeId(GUE.Typenames.Pair, DeclParams)) };
-		FStructSchemaId BindId = GUE.Names.IndexStruct(BindType);
-		FStructSchemaId DeclId = GUE.Names.IndexStruct(DeclType);
+		FStructId BindId = GUE.Names.IndexStruct(BindType);
+		FStructId DeclId = GUE.Names.IndexStruct(DeclType);
 		FMemberId Members[2] = { GUE.Members.Key, GUE.Members.Value };
 
 		// Todo: Ownership / memory leak
@@ -2704,7 +2704,7 @@ public:
 	inline FMemberBinding Bind(FMemberBinding Key, FMemberBinding Value)
 	{
 		FMemberBinding Out(0);
-		Out.InnermostSchema = BindImpl(FPair{Key, Value});
+		Out.InnermostSchema = FInnerId(BindImpl(FPair{Key, Value}));
 		Out.InnermostType = DefaultStructBindType;
 		return Out;
 	}
@@ -2950,7 +2950,7 @@ struct FStaticArrayBinding : IItemRangeBinding
 class FPropertyBinder
 {
 public:
-	FPropertyBinder(FStructSchemaId Id) : Owner(Id) {}
+	FPropertyBinder(FStructId Id) : Owner(Id) {}
 
 	void								AddMember(FMemberBinding Member)	{ Members.Add(Member); }
 	void								BindMember(FProperty* Property)		{ Members.Emplace(BindProperty(Property)); }
@@ -2960,7 +2960,7 @@ public:
 	bool								IsDense() { return GetOwnerOccupancy() == EMemberPresence::RequireAll; }
 
 private:
-	const FStructSchemaId							Owner;
+	const FStructId							Owner;
 	TOptional<FScopeId>								OwnerScope;
 	TOptional<EMemberPresence>						OwnerOccupancy;
 	TPagedArray<FRangeBinding, 1024>				Ranges;
@@ -3036,10 +3036,10 @@ private:
 		return Out;
 	}
 
-	inline FMemberBinding BindAsStruct(FProperty* Property, FStructSchemaId Id)
+	inline FMemberBinding BindAsStruct(FProperty* Property, FStructId Id)
 	{
 		FMemberBinding Out(Property->GetOffset_ForInternal());
-		Out.InnermostSchema = Id;
+		Out.InnermostSchema = FInnerId(Id);
 		Out.InnermostType = DefaultStructBindType;
 		return Out;
 	}
@@ -3069,17 +3069,17 @@ private:
 	inline FMemberBinding BindEnum(FEnumProperty* Property)
 	{
 		FMemberBinding Out(Property->GetOffset_ForInternal());
-		Out.InnermostSchema = GUE.Names.IndexEnum(IndexType(Property->GetEnum()));
+		Out.InnermostSchema = FInnerId(GUE.Names.IndexEnum(IndexType(Property->GetEnum())));
 		FUnpackedLeafType Leaf = {ELeafType::Enum, WidthOf(Property->GetElementSize())};
 		Out.InnermostType = FMemberBindType(Leaf);
 		return Out;
 	}
 
-	inline FLeafBindType BindByte(FByteProperty* Property, FOptionalSchemaId& OutEnumId)
+	inline FLeafBindType BindByte(FByteProperty* Property, FOptionalInnerId& OutEnumId)
 	{
 		if (const UEnum* Enum = Property->GetIntPropertyEnum())
 		{
-			OutEnumId = GUE.Names.IndexEnum(IndexType(Enum));
+			OutEnumId = FInnerId(GUE.Names.IndexEnum(IndexType(Enum)));
 			return FLeafBindType(ELeafBindType::Enum, ELeafWidth::B8);
 		}
 		return FLeafBindType(ELeafBindType::IntU, ELeafWidth::B8);
@@ -3114,7 +3114,7 @@ private:
 		FMemberBinding Key = BindSingleProperty(Property->KeyProp);
 		FMemberBinding Value = BindSingleProperty(Property->ValueProp);
 		FMemberBinding Pair = GPairs.Bind(Key, Value);
-		FStructSchemaId PairId = static_cast<FStructSchemaId>(Pair.InnermostSchema.Get());
+		FStructId PairId = Pair.InnermostSchema.Get().AsStruct();
 
 		bool bFreezable = EnumHasAnyFlags(Property->MapFlags, EMapPropertyFlags::UsesMemoryImageAllocator);
 		return bFreezable	? BindAsRange(Property, GMaps.BindFreezable(Property, PairId, {Key, Value, Pair}), Pair)
@@ -3178,7 +3178,7 @@ private:
 		}
 		else if (HasAny<MiscMask>(Flags))
 		{
-			FStructSchemaId BindId = HasAny<CASTCLASS_FInterfaceProperty>(Flags)
+			FStructId BindId = HasAny<CASTCLASS_FInterfaceProperty>(Flags)
 				? BindInterface(static_cast<FInterfaceProperty*>(Property))
 				: BindSparseDelegate(Owner, static_cast<FMulticastSparseDelegateProperty*>(Property));
 			return BindAsStruct(Property, BindId);
@@ -3250,7 +3250,7 @@ private:
 		{
 			// Create struct binding to allow delta serialization
 			FType StaticArrayType = MakeStaticArrayTypename(Property->GetFName());
-			FStructSchemaId StaticArrayId = GUE.Names.IndexStruct(StaticArrayType);
+			FStructId StaticArrayId = GUE.Names.IndexStruct(StaticArrayType);
 
 			// Todo: Ownership
 			TConstArrayView<FMemberId> Numerals = GUE.Numerals.MakeRange(IntCastChecked<uint16>(ArrayDim));
@@ -3267,7 +3267,7 @@ private:
 			GUE.Schemas.BindStruct(StaticArrayId, StaticArrayId, Elements);
 			
 			Out.InnermostType = DefaultStructBindType;
-			Out.InnermostSchema = static_cast<FSchemaId>(StaticArrayId);
+			Out.InnermostSchema = FInnerId(StaticArrayId);
 			Out.RangeBindings = {};
 		}
 	
@@ -3299,7 +3299,7 @@ void BindSuperMembers(FPropertyBinder& Out, const UStruct* Struct)
 	}
 }
 
-void BindStruct(FStructSchemaId Id, const UStruct* Struct)
+void BindStruct(FStructId Id, const UStruct* Struct)
 {
 	if (GUE.Customs.FindStruct(Id))
 	{
@@ -3320,7 +3320,7 @@ void BindStruct(FStructSchemaId Id, const UStruct* Struct)
 			{
 				FMemberBinding Member;
 				Member.InnermostType = SuperStructBindType;
-				Member.InnermostSchema = static_cast<FSchemaId>(Declared.Super.Get());
+				Member.InnermostSchema = FInnerId(Declared.Super.Get());
 				Binder.AddMember(Member);
 			}
 		}
@@ -3343,15 +3343,15 @@ static void BindInitialTypes()
 	UE_LOGFMT(LogPlainPropsUObject, Display, "Binding types to PlainProps schemas...");
 
 	// Declare all UScriptStruct/UClass/UFunction/UEnum
-	static const FStructSchemaId SkipStructs[] = { GUE.Structs.VerseFunction };
-	TArray<FStructSchemaId> Ids;
+	static const FStructId SkipStructs[] = { GUE.Structs.VerseFunction };
+	TArray<FStructId> Ids;
 	for (TObjectIterator<UField> It; It; ++It)
 	{
 		UField* Field = *It;
 		if (UStruct* Struct = Cast<UStruct>(Field))
 		{
 			FType Type = IndexType(Struct);
-			FStructSchemaId Id = GUE.Names.IndexStruct(Type);
+			FStructId Id = GUE.Names.IndexStruct(Type);
 			if (!Algo::Find(SkipStructs, Id))
 			{
 				DeclareStruct(Struct, Type, Id);
@@ -3365,7 +3365,7 @@ static void BindInitialTypes()
 	}
 
 	// Bind all UScriptStruct/UClass/UFunction
-	const FStructSchemaId* IdIt = &Ids[0];
+	const FStructId* IdIt = &Ids[0];
 	for (TObjectIterator<UStruct> It; It; ++It, ++IdIt)
 	{
 		if (!Algo::Find(SkipStructs, *IdIt))
@@ -3376,7 +3376,7 @@ static void BindInitialTypes()
 	check(IdIt == &Ids[0] + Ids.Num());
 }
 
-static void DeclareProperty(FStructSchemaId Id, TConstArrayView<FMemberId> Members, EMemberPresence Occupancy)
+static void DeclareProperty(FStructId Id, TConstArrayView<FMemberId> Members, EMemberPresence Occupancy)
 {
 	GUE.Types.DeclareStruct(Id,	GUE.Names.Resolve(Id), Members, Occupancy);
 }
@@ -3549,7 +3549,7 @@ struct TCustomPropertyBindings
 		Bind(GUE.Structs.LazyObjectPtr, LazyObjectPtr);
 	}
 
-	void Bind(FStructSchemaId Id, ICustomBinding& Binding)
+	void Bind(FStructId Id, ICustomBinding& Binding)
 	{
 		Overlay.BindStruct(Id, Binding, GUE.Types.Get(Id));
 	}
@@ -3568,7 +3568,7 @@ struct TCustomPropertyBindings
 struct FMemoryBatch
 {
 	TArray64<uint8>				Data;
-	TArray<FStructSchemaId>		RuntimeIds; // To avoid reindexing schema FType
+	TArray<FStructId>		RuntimeIds; // To avoid reindexing schema FType
 	FMemoryPropertyBatch		Properties; // Contains FTexts referenced from Data
 };
 
@@ -3585,13 +3585,13 @@ public:
 		SavedObjects.Reserve(NumReserve);
 	}
 
-	FORCENOINLINE void Save(FStructSchemaId Id, const UObject* Object, const UObject* Arch)
+	FORCENOINLINE void Save(FStructId Id, const UObject* Object, const UObject* Arch)
 	{
 		FBuiltStruct* Built = Arch ? SaveStructDelta(Object, Arch, Id, DeltaCtx) : SaveStruct(Object, Id, FlatCtx);
 		SavedObjects.Emplace(Id, Built, Object);
 	}
 
-	TArray64<uint8>	Write(TArray<FStructSchemaId>& OutRuntimeIds) const
+	TArray64<uint8>	Write(TArray<FStructId>& OutRuntimeIds) const
 	{
 		// Build partial schemas
 		const FStructBindIds BindIds(GUE.Customs, GUE.Schemas);
@@ -3638,7 +3638,7 @@ public:
 private:
 	struct FSavedObject
 	{
-		FStructSchemaId		Id;
+		FStructId		Id;
 		FBuiltStruct*		Built;
 		const UObject*		Input; // For debug
 	};
@@ -3659,7 +3659,7 @@ private:
 class FMemoryBatchLoader
 {
 public:
-	FMemoryBatchLoader(const FCustomBindings& Customs, FMemoryView Data, TConstArrayView<FStructSchemaId> RuntimeIds)
+	FMemoryBatchLoader(const FCustomBindings& Customs, FMemoryView Data, TConstArrayView<FStructId> RuntimeIds)
 	{
 		//// Read ids
 		
@@ -3832,7 +3832,7 @@ struct FInstance
 	UObject* PP;
 	UObject* TPS;
 	UObject* UPS;
-	FStructSchemaId Id;
+	FStructId Id;
 
 	void Init()
 	{
@@ -3987,7 +3987,7 @@ static int32 Roundtrip()
 		if (DiffStructs(Instance.Id, Instance.Orig, Instance.PP, /* in-out */ DiffCtx))
 		{
 			GPpDiff.Last = &Instance;
-			PrintDiff(/* out */ PpDiffs, DiffCtx.Out, GUE.Names);
+			PrintDiff(/* out */ PpDiffs, GUE.Names, DiffCtx.Out);
 			Batch.Reload(Instance.UPS, &Instance - Instances.GetData());
 			DiffCtx.Out.Reset();
 			PpDiffs.Append(" in ");
@@ -4031,7 +4031,7 @@ static int32 Roundtrip()
 		if (DiffStructs(Instance.Id, Instance.Orig, Instance.TPS, /* in-out */ DiffCtx))
 		{
 			GTpsDiff.Last = &Instance;
-			PrintDiff(/* out */ TpsDiffs, DiffCtx.Out, GUE.Names);
+			PrintDiff(/* out */ TpsDiffs, GUE.Names, DiffCtx.Out);
 			
 			DiffCtx.Out.Reset();
 			TpsDiffs.Append(" in ");
@@ -4081,7 +4081,7 @@ static int32 Roundtrip()
 		if (DiffStructs(Instance.Id, Instance.Orig, Instance.UPS, /* in-out */ DiffCtx))
 		{
 			GUpsDiff.Last = &Instance;
-			PrintDiff(/* out */ UpsDiffs, DiffCtx.Out, GUE.Names);
+			PrintDiff(/* out */ UpsDiffs, GUE.Names, DiffCtx.Out);
 			
 			DiffCtx.Out.Reset();
 			UpsDiffs.Append(" in ");
@@ -4179,6 +4179,6 @@ extern FString _print(const PlainProps::FOptionalMemberId& Id)		{ return PlainPr
 extern FString _print(const PlainProps::FScopeId& Id)				{ return PlainProps::UE::GUE.Debug.Print(Id); }
 extern FString _print(const PlainProps::FTypenameId& Id)			{ return PlainProps::UE::GUE.Debug.Print(Id); }
 extern FString _print(const PlainProps::FConcreteTypenameId& Id)	{ return PlainProps::UE::GUE.Debug.Print(Id); }
-extern FString _print(const PlainProps::FType& Id)				{ return PlainProps::UE::GUE.Debug.Print(Id); }
-extern FString _print(const PlainProps::FEnumSchemaId& Id)			{ return PlainProps::UE::GUE.Debug.Print(Id); }
-extern FString _print(const PlainProps::FStructSchemaId& Id)		{ return PlainProps::UE::GUE.Debug.Print(Id); }
+extern FString _print(const PlainProps::FType& Id)					{ return PlainProps::UE::GUE.Debug.Print(Id); }
+extern FString _print(const PlainProps::FEnumId& Id)				{ return PlainProps::UE::GUE.Debug.Print(Id); }
+extern FString _print(const PlainProps::FStructId& Id)				{ return PlainProps::UE::GUE.Debug.Print(Id); }
\ No newline at end of file
diff --git a/Engine/Plugins/Experimental/PlainPropsUObject/Source/Public/PlainPropsUObjectRuntime.h b/Engine/Plugins/Experimental/PlainPropsUObject/Source/Public/PlainPropsUObjectRuntime.h
--- a/Engine/Plugins/Experimental/PlainPropsUObject/Source/Public/PlainPropsUObjectRuntime.h
+++ b/Engine/Plugins/Experimental/PlainPropsUObject/Source/Public/PlainPropsUObjectRuntime.h
@@ -25,16 +25,16 @@ class FDefaultStructs final : public IDefaultStructs
 {
 public:
 	~FDefaultStructs();
-	void									Bind(FStructSchemaId Id, const UScriptStruct* Struct);
-	void									BindZeroes(FStructSchemaId Id, SIZE_T Size, uint32 Alignment);
-	void									BindStatic(FStructSchemaId Id, const void* Struct);
-	void									Drop(FStructSchemaId Id);
-	virtual const void*						Get(FStructSchemaId Id) override;
+	void									Bind(FStructId Id, const UScriptStruct* Struct);
+	void									BindZeroes(FStructId Id, SIZE_T Size, uint32 Alignment);
+	void									BindStatic(FStructId Id, const void* Struct);
+	void									Drop(FStructId Id);
+	virtual const void*						Get(FStructId Id) override;
 
 private:
 	FMutableMemoryView						Zeroes;
 	TBitArray<>								Instanced;
-	TMap<FStructSchemaId, FDefaultInstance>	Instances;
+	TMap<FStructId, FDefaultInstance>		Instances;
 #if DO_CHECK
 	TBitArray<>								Bound;
 #endif
@@ -95,24 +95,24 @@ struct FCommonStructIds
 {
 	explicit FCommonStructIds(const FCommonScopeIds& Scopes, TIdIndexer<FName>& Names);
 	
-	FStructSchemaId				Name;
-	FStructSchemaId				Text;
-	FStructSchemaId				Guid;
-	FStructSchemaId				FieldPath;
-	FStructSchemaId				SoftObjectPath;
-	FStructSchemaId				ClassPtr;
-	FStructSchemaId				ObjectPtr;
-	FStructSchemaId				WeakObjectPtr;
-	FStructSchemaId				LazyObjectPtr;
-	FStructSchemaId				SoftObjectPtr;
-	FStructSchemaId				ScriptInterface;
-	FStructSchemaId				Delegate;
-	FStructSchemaId				MulticastDelegate;
-	FStructSchemaId				MulticastInlineDelegate;
-	FStructSchemaId				MulticastSparseDelegate;
-	FStructSchemaId				VerseFunction;
-	FStructSchemaId				DynamicallyTypedValue;
-	FStructSchemaId				ReferencePropertyValue;
+	FStructId					Name;
+	FStructId					Text;
+	FStructId					Guid;
+	FStructId					FieldPath;
+	FStructId					SoftObjectPath;
+	FStructId					ClassPtr;
+	FStructId					ObjectPtr;
+	FStructId					WeakObjectPtr;
+	FStructId					LazyObjectPtr;
+	FStructId					SoftObjectPtr;
+	FStructId					ScriptInterface;
+	FStructId					Delegate;
+	FStructId					MulticastDelegate;
+	FStructId					MulticastInlineDelegate;
+	FStructId					MulticastSparseDelegate;
+	FStructId					VerseFunction;
+	FStructId					DynamicallyTypedValue;
+	FStructId					ReferencePropertyValue;
 };
 
 struct FCommonMemberIds
@@ -156,8 +156,8 @@ struct FRuntimeIds
 	static FMemberId			IndexMember(FAnsiStringView Name)		{ return GUE.Names.NameMember(FName(Name)); }
 	static FConcreteTypenameId	IndexTypename(FAnsiStringView Name)		{ return GUE.Names.NameType(FName(Name)); }
 	static FFlatScopeId			IndexScope(FAnsiStringView Name)		{ return GUE.Names.NameScope(FName(Name)); }
-	static FEnumSchemaId		IndexEnum(FType Type)					{ return GUE.Names.IndexEnum(Type); }
-	static FStructSchemaId		IndexStruct(FType Type)					{ return GUE.Names.IndexStruct(Type); }
+	static FEnumId				IndexEnum(FType Type)					{ return GUE.Names.IndexEnum(Type); }
+	static FStructId			IndexStruct(FType Type)					{ return GUE.Names.IndexStruct(Type); }
 	static FIdIndexerBase&		GetIndexer()							{ return GUE.Names; }
 };
 
@@ -269,9 +269,9 @@ template<> struct TCustomBind<FReferencePropertyValue> { using Type = UE::FRefer
 template<> struct TCustomBind<FVerseFunction> { using Type = UE::FVerseFunctionBinding; };
 
 // Temporary way to tie certain types to /Script/Core scope
-template<> FScopeId IndexNamespaceId<UE::FRuntimeIds, TTypename<FTransform>>()		{ return UE::GUE.Scopes.Core; }
-template<> FScopeId IndexNamespaceId<UE::FRuntimeIds, TTypename<FGuid>>()			{ return UE::GUE.Scopes.Core; }
-template<> FScopeId IndexNamespaceId<UE::FRuntimeIds, TTypename<FColor>>()			{ return UE::GUE.Scopes.Core; }
-template<> FScopeId IndexNamespaceId<UE::FRuntimeIds, TTypename<FLinearColor>>()	{ return UE::GUE.Scopes.Core; }
+template<> inline FScopeId IndexNamespaceId<UE::FRuntimeIds, TTypename<FTransform>>()		{ return UE::GUE.Scopes.Core; }
+template<> inline FScopeId IndexNamespaceId<UE::FRuntimeIds, TTypename<FGuid>>()			{ return UE::GUE.Scopes.Core; }
+template<> inline FScopeId IndexNamespaceId<UE::FRuntimeIds, TTypename<FColor>>()			{ return UE::GUE.Scopes.Core; }
+template<> inline FScopeId IndexNamespaceId<UE::FRuntimeIds, TTypename<FLinearColor>>()		{ return UE::GUE.Scopes.Core; }
 
-}
+}
\ No newline at end of file
