PlainProps: Add missing bounds checking to FDebugIds::Print(FType).

#rnx

diff --git a/Engine/Plugins/Experimental/PlainProps/Source/Private/PlainPropsPrint.cpp b/Engine/Plugins/Experimental/PlainProps/Source/Private/PlainPropsPrint.cpp
--- a/Engine/Plugins/Experimental/PlainProps/Source/Private/PlainPropsPrint.cpp
+++ b/Engine/Plugins/Experimental/PlainProps/Source/Private/PlainPropsPrint.cpp
@@ -1321,11 +1321,19 @@ FString FDebugIds::Print(FOptionalMemberId Name) const
 	return FString(FStringView(StringCast<TCHAR>(Out.GetData(), Out.Len())));
 }
 
+static const bool IsValidScope(FScopeId Scope, const FIds& Ids)
+{
+	if (Scope.IsFlat())
+	{
+		return Scope.AsFlat().Name.Idx < Ids.NumNames();
+	}
+	return Scope.AsNested().Idx < Ids.NumNestedScopes();
+}
+
 FString FDebugIds::Print(FScopeId Scope) const
 {
 	TUtf8StringBuilder<128> Out;
-	if ((Scope.IsFlat() && Scope.AsFlat().Name.Idx < Ids.NumNames()) ||
-		Scope.AsNested().Idx < Ids.NumNestedScopes())
+	if (IsValidScope(Scope, Ids))
 	{
 		Ids.AppendString(Out, Scope);
 	}
@@ -1336,11 +1344,19 @@ FString FDebugIds::Print(FScopeId Scope) const
 	return FString(FStringView(StringCast<TCHAR>(Out.GetData(), Out.Len())));
 }
 
+static const bool IsValidTypename(FTypenameId Typename, const FIds& Ids)
+{
+	if (Typename.IsConcrete())
+	{
+		return Typename.AsConcrete().Id.Idx < Ids.NumNames();
+	}
+	return Typename.AsParametric().Idx < Ids.NumNames();
+}
+
 FString FDebugIds::Print(FTypenameId Typename) const
 {
 	TUtf8StringBuilder<128> Out;
-	if ((Typename.IsConcrete() && Typename.AsConcrete().Id.Idx < Ids.NumNames()) ||
-		Typename.AsParametric().Idx < Ids.NumNames())
+	if (IsValidTypename(Typename, Ids))
 	{
 		Ids.AppendString(Out, Typename);
 	}
@@ -1364,7 +1380,14 @@ FString FDebugIds::Print(FParametricTypeId Typename) const
 FString FDebugIds::Print(FType Type) const
 {
 	TUtf8StringBuilder<128> Out;
-	Ids.AppendString(Out, Type);
+	if (IsValidScope(Type.Scope, Ids) && IsValidTypename(Type.Name, Ids))
+	{
+		Ids.AppendString(Out, Type);
+	}
+	else
+	{
+		Out << GLiterals.Oob;
+	}
 	return FString(FStringView(StringCast<TCHAR>(Out.GetData(), Out.Len())));
 }
 
