PlainProps: Setup proper FBatchIds instances for all various test phases.

#rnx

diff --git a/Engine/Plugins/Experimental/PlainProps/Source/Tests/PlainPropsTests.cpp b/Engine/Plugins/Experimental/PlainProps/Source/Tests/PlainPropsTests.cpp
--- a/Engine/Plugins/Experimental/PlainProps/Source/Tests/PlainPropsTests.cpp
+++ b/Engine/Plugins/Experimental/PlainProps/Source/Tests/PlainPropsTests.cpp
@@ -305,11 +305,11 @@ private:
 	TArray<const char*> Names;
 };
 
-class FTestBatchNames final : public FStableBatchIds
+class FTestBatchIds final : public FStableBatchIds
 {
 	const FTestNameReader& Names;
 public:
-	FTestBatchNames(const FTestNameReader& InNames, FReadBatchId Batch) : FStableBatchIds(Batch), Names(InNames) {}
+	FTestBatchIds(const FTestNameReader& InNames, FReadBatchId Batch) : FStableBatchIds(Batch), Names(InNames) {}
 	
 	using FStableBatchIds::AppendString;
 	virtual uint32 NumNames() const override { return Names.NumNames(); }
@@ -339,12 +339,11 @@ public:
 	return Batch;
 }
 
-static void RoundtripText(const FTestNameReader& Names, FReadBatchId Batch, TConstArrayView<FStructView> Objects)
+static void RoundtripText(const FBatchIds& BatchIds, TConstArrayView<FStructView> Objects)
 {
 	// Print yaml
 	TUtf8StringBuilder<4096> Yaml;
-	FTestBatchNames BatchNames(Names, Batch);
-	PrintYamlBatch(Yaml, BatchNames, Objects);
+	PrintYamlBatch(Yaml, BatchIds, Objects);
 	FUtf8StringView YamlView = Yaml.ToView();
 
 	// Log yaml
@@ -357,7 +356,7 @@ static void RoundtripText(const FTestNameReader& Names, FReadBatchId Batch, TCon
 	FReadBatchId ParsedBatch = ParseBatch(Data, ParsedObjects, YamlView);
 
 	// Diff schemas
-	CHECK_FALSE(DiffSchemas(Batch, ParsedBatch));
+	CHECK_FALSE(DiffSchemas(BatchIds.GetBatchId(), ParsedBatch));
 
 	// Diff objects
 	CHECK(Objects.Num() == ParsedObjects.Num());
@@ -370,7 +369,7 @@ static void RoundtripText(const FTestNameReader& Names, FReadBatchId Batch, TCon
 		if (DiffStruct(In, Parsed, DiffPath))
 		{
 			TUtf8StringBuilder<256> Diff;
-			PrintDiff(Diff, BatchNames, DiffPath);
+			PrintDiff(Diff, BatchIds, DiffPath);
 			FAIL_CHECK(FString::Printf(TEXT("Diff in '%s' in Objects[%d]"), *Print(Diff.ToString()), I));
 		}
 	}
@@ -410,11 +409,13 @@ public:
 		
 		CHECK(It.Grab<uint32>() == TestMagics[4]);
 		CHECK(!Objects.IsEmpty());
+
+		BatchIds.Emplace(Names, Batch);
 	}
 
 	void RoundtripText()
 	{
-		PlainProps::RoundtripText(Names, Objects[0].Schema.Batch, Objects);
+		PlainProps::RoundtripText(BatchIds.GetValue(), Objects);
 	}
 
 	~FTestBatchReader()
@@ -424,9 +425,11 @@ public:
 
 	TConstArrayView<FStructView>		GetObjects() const { return Objects; }
 	const FTestNameReader&				GetNames() const { return Names; }
+	const FTestBatchIds&				GetBatchIds() const { return BatchIds.GetValue(); }
 	
 private:
 	FTestNameReader Names;
+	TOptional<FTestBatchIds> BatchIds;
 	TArray<FStructView> Objects;
 };
 
diff --git a/Engine/Plugins/Experimental/PlainProps/Source/Tests/PlainPropsUeCoreTests.cpp b/Engine/Plugins/Experimental/PlainProps/Source/Tests/PlainPropsUeCoreTests.cpp
--- a/Engine/Plugins/Experimental/PlainProps/Source/Tests/PlainPropsUeCoreTests.cpp
+++ b/Engine/Plugins/Experimental/PlainProps/Source/Tests/PlainPropsUeCoreTests.cpp
@@ -266,36 +266,6 @@ TArray64<uint8> FBatchSaver::Write(TArray<FStructId>* OutMemoryIds) const
 
 ////////////////////////////////////////////////////////////////////////////////////////////////
 
-//class FSchemaBatchIds : public FIds
-//{
-//public:
-//	FSchemaBatchIds(const FSchemaBatch& InSchemas) : Schemas(InSchemas) {}
-//
-//	virtual uint32				NumEnums() const override final						{ return Schemas.NumSchemas - Schemas.NumStructSchemas; }
-//	virtual uint32				NumStructs() const override final					{ return Schemas.NumStructSchemas; }
-//
-//	virtual FType				Resolve(FEnumSchemaId Id) const override final		{ return ResolveEnumSchema(Schemas, Id).Type; }
-//	virtual FType				Resolve(FStructSchemaId Id) const override final	{ return ResolveStructSchema(Schemas, Id).Type; }
-//
-//protected:
-//	const FSchemaBatch& Schemas;
-//};
-//
-//class FMemoryBatchIds final : public FSchemaBatchIds
-//{
-//public:
-//	using FSchemaBatchIds::FSchemaBatchIds;
-//	using FSchemaBatchIds::Resolve;
-//
-//	virtual uint32				NumNames() const override final						{ return GNames.NumNames(); }
-//	virtual uint32				NumNestedScopes() const override final				{ return GNames.NumNestedScopes(); }
-//	virtual uint32				NumParametricTypes() const override final			{ return GNames.NumParametricTypes(); }
-//
-//	virtual FNestedScope		Resolve(FNestedScopeId Id) const override			{ return GNames.Resolve(Id); }
-//	virtual FParametricTypeView	Resolve(FParametricTypeId Id) const override		{ return GNames.Resolve(Id); }
-//	virtual void				AppendString(FUtf8Builder& Out, FNameId Name) const override { GNames.ResolveName(Name).AppendString(Out); }
-//};
-
 class FStableNameBatchIds final : public FStableBatchIds
 {
 	TConstArrayView<FName> Names;
@@ -307,14 +277,14 @@ public:
 	virtual void				AppendString(FUtf8Builder& Out, FNameId Name) const override { Names[Name.Idx].AppendString(Out); }
 };
 
-class FLoadBatchNames final : public FBatchIds
+class FTranslationBatchIds final : public FBatchIds
 {
-	const FIdBinding& Binding;
+	FIdBinding Binding;
 public:
 	using FBatchIds::Resolve;
 	using FBatchIds::AppendString;
 
-	FLoadBatchNames(const FIdBinding& InBinding, FReadBatchId Batch) : FBatchIds(Batch), Binding(InBinding) {}
+	FTranslationBatchIds(FReadBatchId Batch, FIdBinding&& InBinding) : FBatchIds(Batch), Binding(InBinding) {}
 
 	virtual uint32				NumNames() const override												{ return GNames.NumNames(); }
 	virtual void				AppendString(FUtf8Builder& Out, FNameId Name) const override			{ PlainProps::AppendString(Out, GNames.ResolveName(Name)); }
@@ -327,29 +297,11 @@ public:
 	virtual FParametricTypeView	Resolve(FParametricTypeId Id) const override		{ return GNames.Resolve(Binding.Remap(Id)); }
 };
 
-static void RoundtripText(TConstArrayView<FName> Names, const FSchemaBatch& Schemas, TArray<FStructView>&& Objects)
+static void RoundtripText(const FBatchIds& BatchIds, TConstArrayView<FStructView> Objects, ESchemaFormat Format)
 {
-	ESchemaFormat Format = Names.IsEmpty() ? ESchemaFormat::InMemoryNames : ESchemaFormat::StableNames;
-
-	// Mount input schemas and update copy of objects array accordingly
-	FReadBatchId Batch = MountReadSchemas(&Schemas);
-	for (FStructView& Struct : Objects)
-	{
-		Struct.Schema.Batch = Batch;
-	}
-
 	// Print yaml
 	TUtf8StringBuilder<4096> Yaml;
-	if (Format == ESchemaFormat::InMemoryNames)
-	{
-		FMemoryBatchIds Ids(Batch, GNames);
-		PrintYamlBatch(Yaml, Ids, Objects);
-	}
-	else
-	{
-		FStableNameBatchIds Ids(Batch, Names);
-		PrintYamlBatch(Yaml, Ids, Objects);
-	}
+	PrintYamlBatch(Yaml, BatchIds, Objects);
 	FUtf8StringView YamlView = Yaml.ToView();
 
 	// Log yaml
@@ -365,10 +317,8 @@ static void RoundtripText(TConstArrayView<FName> Names, const FSchemaBatch& Sche
 
 	if (Format == ESchemaFormat::StableNames)
 	{
-		FStableNameBatchIds Ids(Batch, Names);
-
 		// Diff schemas
-		CHECK_FALSE(DiffSchemas(Batch, ParsedBatch));
+		CHECK_FALSE(DiffSchemas(BatchIds.GetBatchId(), ParsedBatch));
 
 		// Diff objects
 		CHECK(Objects.Num() == ParsedObjects.Num());
@@ -381,24 +331,22 @@ static void RoundtripText(TConstArrayView<FName> Names, const FSchemaBatch& Sche
 			if (DiffStruct(In, Parsed, DiffPath))
 			{
 				TUtf8StringBuilder<256> Diff;
-				PrintDiff(Diff, Ids, DiffPath);
+				PrintDiff(Diff, BatchIds, DiffPath);
 				FAIL_CHECK(FString::Printf(TEXT("Diff in '%s' in Objects[%d]"), *Print(Diff.ToString()), I));
 			}
 		}
 	}
 
-	// Unmount input schemas and parsed schemas
-	UnmountReadSchemas(Batch);
+	// Unmount parsed schemas
 	UnmountReadSchemas(ParsedBatch);
 }
 
-static void PrintLoadStruct(const FIdBinding& Binding, FStructView StructView, int32 LoadIdx)
+static void PrintLoadStruct(const FBatchIds& BatchIds, FStructView StructView, int32 LoadIdx)
 {
 	TUtf8StringBuilder<4096> YamlString;
 	{
 		FYamlBuilderPtr YamlBuilder = MakeYamlBuilder(YamlString);
-		FLoadBatchNames Ids(Binding, StructView.Schema.Batch);
-		FBatchPrinter Printer(*YamlBuilder, Ids);
+		FBatchPrinter Printer(*YamlBuilder, BatchIds);
 		Printer.PrintObjects({StructView});
 	}
 	FUtf8StringView YamlView = YamlString.ToView();
@@ -413,14 +361,13 @@ class FBatchLoader
 public:
 	FBatchLoader(FMemoryView Data, const FCustomBindings& Underlay, TConstArrayView<FStructId> InRuntimeIds)
 	: Customs(Underlay)
+	, Format(InRuntimeIds.IsEmpty() ? ESchemaFormat::StableNames : ESchemaFormat::InMemoryNames)
 	{
-		ESchemaFormat Format = InRuntimeIds.IsEmpty() ? ESchemaFormat::StableNames : ESchemaFormat::InMemoryNames;
-
 		// Read ids
 		FByteReader It(Data);
 		CHECK(It.Grab<uint32>() == Magics[0]);
-		Ids = GrabNumAndArray<FName>(It);
-		CHECK(Ids.IsEmpty() != InRuntimeIds.IsEmpty());
+		Names = GrabNumAndArray<FName>(It);
+		CHECK(Names.IsEmpty() != InRuntimeIds.IsEmpty());
 		
 		// Read schemas
 		CHECK(It.Grab<uint32>() == Magics[1]);
@@ -436,12 +383,9 @@ public:
 		if (InRuntimeIds.IsEmpty())
 		{
 			// Bind saved ids to runtime ids, make new schemas with new ids and mount them
-			RuntimeIds.Emplace(GNames, Ids, *SavedSchemas);
-			const FSchemaBatch* LoadSchemas = CreateTranslatedSchemas(*SavedSchemas, RuntimeIds->Translation);
-			Batch = MountReadSchemas(LoadSchemas);
-
-			// Store mappings required for structs printing of untranslated scopes and parametric types
-			IdBinding.Emplace(RuntimeIds->Translation);
+			RuntimeIds.Emplace(GNames, Names, *SavedSchemas);
+			TranslatedSchemas = CreateTranslatedSchemas(*SavedSchemas, RuntimeIds->Translation);
+			Batch = MountReadSchemas(TranslatedSchemas);
 		}
 		else
 		{
@@ -465,30 +409,57 @@ public:
 		CHECK(!Objects.IsEmpty());
 
 		// Read names and bind custom loader
-		Names.Names.Append(GrabNumAndArray<FName>(It));
-		Customs.BindStruct(Names.Declaration.Id, Names, GTypes.Get(Names.Declaration.Id), {});
+		NameBinding.Names.Append(GrabNumAndArray<FName>(It));
+		Customs.BindStruct(NameBinding.Declaration.Id, NameBinding, GTypes.Get(NameBinding.Declaration.Id), {});
 		CHECK(It.Grab<uint32>() == Magics[5]);
 
 		// Finally create load plans
 		TConstArrayView<FStructId> LoadStructIds = RuntimeIds ? RuntimeIds->Translation.GetStructIds(SavedSchemas->NumStructSchemas) : InRuntimeIds;
 		Plans = CreateLoadPlans(Batch, GTypes, Customs, GSchemas, LoadStructIds, Format);
 
+		// Create BatchIds for the loading phase
+		if (RuntimeIds)
+		{
+			LoadBatchIds = MakeUnique<FTranslationBatchIds>(Batch, MoveTemp(RuntimeIds->Translation));
+		}
+		else
+		{
+			LoadBatchIds = MakeUnique<FMemoryBatchIds>(Batch, GNames);
+		}
 	}
 
 	~FBatchLoader()
 	{
 		CHECK(LoadIdx == Objects.Num()); // Test should load all saved objects
 		Plans.Reset();
-		const FSchemaBatch* LoadSchemas = UnmountReadSchemas(Objects[0].Schema.Batch);
-		if (IdBinding)
+		UnmountReadSchemas(Objects[0].Schema.Batch);
+		if (TranslatedSchemas)
 		{
-			DestroyTranslatedSchemas(LoadSchemas);
+			DestroyTranslatedSchemas(TranslatedSchemas);
 		}
 	}
 
 	void RoundtripText() const
 	{
-		PlainProps::UE::Test::RoundtripText(Ids, *SavedSchemas, CopyTemp(Objects));
+		if (Format == ESchemaFormat::StableNames)
+		{
+			// Mount and use the saved schemas rather than the translated load schemas
+			// in order to verify that text roundtripping creates identical serialized ids.
+			FReadBatchId Batch = MountReadSchemas(SavedSchemas);
+			TArray<FStructView> StableObjects(Objects);
+			for (FStructView& Struct : StableObjects)
+			{
+				Struct.Schema.Batch = Batch;
+			}
+
+			FStableNameBatchIds StableBatchIds(Batch, Names);
+			PlainProps::UE::Test::RoundtripText(StableBatchIds, StableObjects, Format);
+			UnmountReadSchemas(Batch);
+		}
+		else
+		{
+			PlainProps::UE::Test::RoundtripText(*LoadBatchIds, Objects, Format);
+		}
 	}
 
 	template<class T>
@@ -511,32 +482,32 @@ public:
 	void LoadInto(T& Out)
 	{
 		FStructView In = Objects[LoadIdx];
-		if (IdBinding)
-		{
-			PrintLoadStruct(IdBinding.GetValue(), In, LoadIdx);
-		}
+		PrintLoadStruct(*LoadBatchIds, In, LoadIdx);
 		LoadStruct(&Out, In.Values, In.Schema.Id, *Plans);
 		++LoadIdx;
 	}
 
+	const FBatchIds&			GetBatchIds() { return *LoadBatchIds; }
+
 private:
 	const FSchemaBatch*			SavedSchemas;
-	TConstArrayView<FName>		Ids;
-	FNameBinding				Names;
-	TOptional<FIdBinding>		IdBinding;
+	const FSchemaBatch*			TranslatedSchemas = nullptr;
+	TConstArrayView<FName>		Names;
+	FNameBinding				NameBinding;
+	TUniquePtr<FBatchIds>		LoadBatchIds;
 	FCustomBindingsOverlay		Customs;
 	FLoadBatchPtr				Plans;
 	TArray<FStructView>			Objects;
 	int32						LoadIdx = 0;
+	ESchemaFormat				Format;
 };
 
-
 static void Run(void (*Save)(FBatchSaver&), void (*Load)(FBatchLoader&), const FCustomBindings& Customs = GCustoms)
 {
 	for (ESchemaFormat Format : {ESchemaFormat::StableNames, ESchemaFormat::InMemoryNames})
 	{
 		TArray64<uint8> Data;
-		TArray<FStructId>	RuntimeIds;
+		TArray<FStructId> RuntimeIds;
 		{
 			FBatchSaver Batch(Customs);
 			Save(Batch);
diff --git a/Engine/Plugins/Experimental/PlainPropsUObject/Source/Private/PlainPropsUObjectBindings.cpp b/Engine/Plugins/Experimental/PlainPropsUObject/Source/Private/PlainPropsUObjectBindings.cpp
--- a/Engine/Plugins/Experimental/PlainPropsUObject/Source/Private/PlainPropsUObjectBindings.cpp
+++ b/Engine/Plugins/Experimental/PlainPropsUObject/Source/Private/PlainPropsUObjectBindings.cpp
@@ -3999,29 +3999,18 @@ public:
 }
 
 static void RoundtripText(
-	TConstArrayView<FName> Names,
-	FMemoryView SchemasView,
-	FReadBatchId Batch,
+	const FBatchIds& BatchIds,
 	TConstArrayView<FStructView> Objects,
-	TConstArrayView<FInstance> Instances)
+	TConstArrayView<FInstance> Instances,
+	ESchemaFormat Format)
 {
 	check(Objects.Num() == Instances.Num());
-	ESchemaFormat Format = Names.IsEmpty() ? ESchemaFormat::InMemoryNames : ESchemaFormat::StableNames;
 
 	// Print yaml
 	UE_LOGFMT(LogPlainPropsUObject, Display, "Printing to PlainProps text using {Format}...", ToString(Format));
 	TUtf8StringBuilder<256> Yaml;
 	Yaml.Reserve(INT_MAX);
-	if (Format == ESchemaFormat::InMemoryNames)
-	{
-		FMemoryBatchIds Ids(Batch, GUE.Names);
-		PrintYamlBatch(Yaml, Ids, Objects);
-	}
-	else
-	{
-		FStableNameBatchIds Ids(Batch, Names);
-		PrintYamlBatch(Yaml, Ids, Objects);
-	}
+	PrintYamlBatch(Yaml, BatchIds, Objects);
 	FUtf8StringView YamlView = Yaml.ToView();
 
 	// Write to file
@@ -4042,10 +4031,9 @@ static void RoundtripText(
 	if (Format == ESchemaFormat::StableNames)
 	{
 		UE_LOGFMT(LogPlainPropsUObject, Display, "Diffing PlainProps parsed objects using {Format}...", ToString(Format));
-		FStableNameBatchIds Ids(Batch, Names);
 
 		// Diff schemas
-		check(!DiffSchemas(Batch, ParsedBatch));
+		check(!DiffSchemas(BatchIds.GetBatchId(), ParsedBatch));
 
 		// Diff objects
 		check(Objects.Num() == ParsedObjects.Num());
@@ -4059,7 +4047,7 @@ static void RoundtripText(
 			FReadDiffPath DiffPath;
 			if (DiffStruct(In, Parsed, DiffPath))
 			{
-				PrintDiff(Diffs, Ids, DiffPath);
+				PrintDiff(Diffs, BatchIds, DiffPath);
 				Diffs.Append(" in ");
 				Diffs.Append(Instances[I].PathName);
 				Diffs.Append("\n");
@@ -4078,25 +4066,26 @@ static void RoundtripText(
 class FBatchTextRoundtripper
 {
 public:
-	FBatchTextRoundtripper(FMemoryView Data, ESchemaFormat Format)
+	FBatchTextRoundtripper(FMemoryView Data, ESchemaFormat InFormat) : Format(InFormat)
 	{
 		FByteReader It(Data);
 
 		// Read FNames when using Stable Names
+		TConstArrayView<FName> Names;
 		if (Format == ESchemaFormat::StableNames)
 		{
 			verify(It.Grab<uint32>() == Magics[0]);
-			Ids = GrabNumAndArray<FName>(It);
+			Names = GrabNumAndArray<FName>(It);
 		}
 		
 		// Read and mount schemas
 		check(It.Grab<uint32>() == Magics[1]);
 		It.SkipAlignmentPadding<uint32>();
 		const uint32 SchemasSize = It.Grab<uint32>();
-		SavedSchemasView = It.GrabSlice(SchemasSize);
+		FMemoryView SavedSchemasView = It.GrabSlice(SchemasSize);
 		const FSchemaBatch* SavedSchemas = ValidateSchemas(SavedSchemasView);
 		check(It.Grab<uint32>() == Magics[2]);
-		Batch = MountReadSchemas(SavedSchemas);
+		FReadBatchId Batch = MountReadSchemas(SavedSchemas);
 		
 		// Read objects
 		while (uint64 NumBytes = It.GrabVarIntU())
@@ -4109,23 +4098,32 @@ public:
 		
 		check(It.Grab<uint32>() == Magics[4]);
 		check(!Objects.IsEmpty());
+
+		// Create BatchIds 
+		if (Format == ESchemaFormat::StableNames)
+		{
+			BatchIds = MakeUnique<FStableNameBatchIds>(Batch, Names);
+		}
+		else
+		{
+			BatchIds = MakeUnique<FMemoryBatchIds>(Batch, GUE.Names);
+		}
 	}
 
 	~FBatchTextRoundtripper()
 	{
-		UnmountReadSchemas(Batch);
+		UnmountReadSchemas(BatchIds->GetBatchId());
 	}
 
 	void RoundtripText(TConstArrayView<FInstance> Instances) const
 	{
-		PlainProps::UE::RoundtripText(Ids, SavedSchemasView, Batch, Objects, Instances);
+		PlainProps::UE::RoundtripText(*BatchIds, Objects, Instances, Format);
 	}
 
 private:
-	FMemoryView					SavedSchemasView;
-	FReadBatchId				Batch;
-	TConstArrayView<FName>		Ids;
 	TArray<FStructView>			Objects;
+	TUniquePtr<FBatchIds>		BatchIds;
+	ESchemaFormat				Format;
 	
 	template<typename T>
 	static TConstArrayView<T> GrabNumAndArray(/* in-out */ FByteReader& It)
