PlainProps: Handle unscoped FScopeId when debug printing FType and FScopeId.

#rnx

diff --git a/Engine/Plugins/Experimental/PlainProps/Source/Private/PlainPropsPrint.cpp b/Engine/Plugins/Experimental/PlainProps/Source/Private/PlainPropsPrint.cpp
--- a/Engine/Plugins/Experimental/PlainProps/Source/Private/PlainPropsPrint.cpp
+++ b/Engine/Plugins/Experimental/PlainProps/Source/Private/PlainPropsPrint.cpp
@@ -1327,7 +1327,11 @@ static const bool IsValidScope(FScopeId Scope, const FIds& Ids)
 	{
 		return Scope.AsFlat().Name.Idx < Ids.NumNames();
 	}
-	return Scope.AsNested().Idx < Ids.NumNestedScopes();
+	else if (Scope)
+	{
+		return Scope.AsNested().Idx < Ids.NumNestedScopes();
+	}
+	return !Scope; // Unscoped
 }
 
 FString FDebugIds::Print(FScopeId Scope) const
