Fixed a use of the deprecated TStringBuilderBase::ToString()

#rnx

diff --git a/Engine/Plugins/Experimental/PlainProps/Source/Private/PlainPropsParse.cpp b/Engine/Plugins/Experimental/PlainProps/Source/Private/PlainPropsParse.cpp
--- a/Engine/Plugins/Experimental/PlainProps/Source/Private/PlainPropsParse.cpp
+++ b/Engine/Plugins/Experimental/PlainProps/Source/Private/PlainPropsParse.cpp
@@ -715,7 +715,7 @@ TOptional<FToken> FYamlTokenizer::GrabToken()
 
 FUtf8StringView FYamlTokenizer::GetLastError() const
 {
-	return LastError.ToString();
+	return LastError;
 }
 
 TOptional<FToken> FYamlTokenizer::GrabBeginDocument()
