PlainProps: Replace TIdIndexer<FUtf8StringView> with  TIdIndexer<FSensitiveUtf8View> for deterministic case sensitive text parsing with better performance.

#rnx

diff --git a/Engine/Plugins/Experimental/PlainProps/Source/Private/PlainPropsInternalParse.h b/Engine/Plugins/Experimental/PlainProps/Source/Private/PlainPropsInternalParse.h
--- a/Engine/Plugins/Experimental/PlainProps/Source/Private/PlainPropsInternalParse.h
+++ b/Engine/Plugins/Experimental/PlainProps/Source/Private/PlainPropsInternalParse.h
@@ -11,6 +11,7 @@
 #include "Misc/Optional.h"
 #include "Templates/Tuple.h"
 #include "Templates/UniquePtr.h"
+#include "Hash/xxhash.h"
 
 namespace PlainProps
 {
@@ -93,8 +94,30 @@ struct FParsedSchemas
 };
 
 ///////////////////////////////////////////////////////////////////////////////
+class FSensitiveUtf8View : public FUtf8StringView
+{
+public:
+	FSensitiveUtf8View(FUtf8StringView View) : FUtf8StringView(View) {}
+
+	friend inline bool operator==(FSensitiveUtf8View Lhs, FSensitiveUtf8View Rhs)
+	{
+		return Lhs.Equals(Rhs, ESearchCase::CaseSensitive);
+	}
+	friend inline uint32 GetTypeHash(FSensitiveUtf8View View)
+	{
+		return static_cast<uint32>(FXxHash64::HashBuffer(View.GetData(), View.Len()).Hash);
+	}
+	friend inline bool operator!=(FSensitiveUtf8View Lhs, FSensitiveUtf8View Rhs)
+	{
+		return !operator==(Lhs, Rhs);
+	}
+	friend inline bool operator<(FSensitiveUtf8View Lhs, FSensitiveUtf8View Rhs)
+	{
+		return Lhs.Compare(Rhs, ESearchCase::IgnoreCase) < 0;
+	}
+};
 
-class FTextIndexer : public TIdIndexer<FUtf8StringView>
+class FTextIndexer : public TIdIndexer<FSensitiveUtf8View>
 {
 public:
 	FOptionalMemberId GetMemberId(FUtf8StringView Name)
@@ -175,4 +198,4 @@ private:
 
 };
 
-} // namespace PlainProps
\ No newline at end of file
+} // namespace PlainProps
diff --git a/Engine/Plugins/Experimental/PlainProps/Source/Private/PlainPropsParse.cpp b/Engine/Plugins/Experimental/PlainProps/Source/Private/PlainPropsParse.cpp
--- a/Engine/Plugins/Experimental/PlainProps/Source/Private/PlainPropsParse.cpp
+++ b/Engine/Plugins/Experimental/PlainProps/Source/Private/PlainPropsParse.cpp
@@ -18,9 +18,9 @@
 namespace PlainProps
 {
 
-// specialization used by FTextIndexer/TIdIndexer<FUtf8StringView>
+// specialization used by FTextIndexer/TIdIndexer<FSensitiveUtf8View>
 template<>
-void AppendString(FUtf8StringBuilderBase& Out, const FUtf8StringView& Str)
+void AppendString(FUtf8StringBuilderBase& Out, const FSensitiveUtf8View& Str)
 {
 	Out.Append(Str);
 }
