PlainProps: First drop of debug visualizer support based on TIdIndexer<NameT>.
Works for layers using the global indexer in tests and commandlet, as well as for the parse specific indexer.
See PlainProps.natvis for some WIP documentation.

#rnx

diff --git a/Engine/Plugins/Experimental/PlainProps/Source/Private/PlainPropsParse.cpp b/Engine/Plugins/Experimental/PlainProps/Source/Private/PlainPropsParse.cpp
--- a/Engine/Plugins/Experimental/PlainProps/Source/Private/PlainPropsParse.cpp
+++ b/Engine/Plugins/Experimental/PlainProps/Source/Private/PlainPropsParse.cpp
@@ -12,6 +12,7 @@
 #include "PlainPropsInternalParse.h"
 #include "PlainPropsInternalPrint.h"
 #include "PlainPropsInternalText.h"
+#include "PlainPropsVisualize.h"
 #include "PlainPropsWrite.h"
 #include <charconv>
 
@@ -1888,6 +1889,7 @@ FTypedRange FMemberParser::ParseUnicodeLeafValueAsRange(
 
 void FBatchParser::Parse(TArray64<uint8>& Out)
 {
+	DbgVis::FIdScope _(Names, "Utf8View");
 	Tokenize();
 	ParseAll();
 	Write(Out);
diff --git a/Engine/Plugins/Experimental/PlainProps/Source/Private/PlainPropsVisualize.cpp b/Engine/Plugins/Experimental/PlainProps/Source/Private/PlainPropsVisualize.cpp
new file mode 100644
--- /dev/null
+++ b/Engine/Plugins/Experimental/PlainProps/Source/Private/PlainPropsVisualize.cpp
@@ -0,0 +1,95 @@
+// Copyright Epic Games, Inc. All Rights Reserved.
+
+#include "PlainPropsVisualize.h"
+#include "PlainPropsInternalBuild.h"
+#include "PlainPropsInternalFormat.h"
+#include "Misc/Guid.h"
+#include "Modules/VisualizerDebuggingState.h"
+#include <type_traits>
+
+#define PP_DBGVIS_VERIFY_TYPE(Type) \
+	static_assert(sizeof(Type) == sizeof(DbgVis::Type)); \
+	static_assert(alignof(Type) == alignof(DbgVis::Type))
+
+#define PP_DBGVIS_VERIFY_MEMBER(Type, Member) \
+	static_assert(offsetof(Type, Member) == offsetof(DbgVis::Type, Member)); \
+	static_assert(std::is_same_v<decltype(Type::Member),decltype(DbgVis::Type::Member)>);
+
+#define PP_DBGVIS_VERIFY_BITFIELD(Type, Member) \
+	static_assert(std::is_integral_v<decltype(Type::Member)> == std::is_integral_v<decltype(DbgVis::Type::Member)>); \
+	static_assert(std::is_enum_v<decltype(Type::Member)> == std::is_enum_v<decltype(DbgVis::Type::Member)>);
+
+#define PP_DBGVIS_VERIFY_ENUMERATOR(Enum, Constant) \
+	static_assert((uint64)Enum::Constant == (uint64)DbgVis::Enum::Constant);
+
+
+namespace PlainProps
+{
+
+// Verify layouts for duplicated debug types for InternalBuild.h
+
+PP_DBGVIS_VERIFY_TYPE(FBuiltStruct);
+PP_DBGVIS_VERIFY_MEMBER(FBuiltStruct, NumMembers);
+PP_DBGVIS_VERIFY_MEMBER(FBuiltStruct, Members);
+
+PP_DBGVIS_VERIFY_TYPE(FBuiltRange);
+PP_DBGVIS_VERIFY_MEMBER(FBuiltRange, Num);
+PP_DBGVIS_VERIFY_MEMBER(FBuiltRange, Data);
+
+
+// Verify layouts for duplicated debug types for InternalFormat.h
+
+PP_DBGVIS_VERIFY_TYPE(FSchemaBatch);
+PP_DBGVIS_VERIFY_MEMBER(FSchemaBatch, NumNestedScopes);
+PP_DBGVIS_VERIFY_MEMBER(FSchemaBatch, NestedScopesOffset);
+PP_DBGVIS_VERIFY_MEMBER(FSchemaBatch, NumParametricTypes);
+PP_DBGVIS_VERIFY_MEMBER(FSchemaBatch, NumSchemas);
+PP_DBGVIS_VERIFY_MEMBER(FSchemaBatch, NumStructSchemas);
+PP_DBGVIS_VERIFY_MEMBER(FSchemaBatch, SchemaOffsets);
+
+PP_DBGVIS_VERIFY_TYPE(ESuper);
+PP_DBGVIS_VERIFY_ENUMERATOR(ESuper, No);
+PP_DBGVIS_VERIFY_ENUMERATOR(ESuper, Unused);
+PP_DBGVIS_VERIFY_ENUMERATOR(ESuper, Used);
+PP_DBGVIS_VERIFY_ENUMERATOR(ESuper, Reused);
+
+PP_DBGVIS_VERIFY_TYPE(FStructSchema);
+PP_DBGVIS_VERIFY_MEMBER(FStructSchema, Type);
+PP_DBGVIS_VERIFY_MEMBER(FStructSchema, NumMembers);
+PP_DBGVIS_VERIFY_MEMBER(FStructSchema, NumRangeTypes);
+PP_DBGVIS_VERIFY_MEMBER(FStructSchema, NumInnerSchemas);
+PP_DBGVIS_VERIFY_BITFIELD(FStructSchema, Inheritance);
+PP_DBGVIS_VERIFY_BITFIELD(FStructSchema, IsDense);
+PP_DBGVIS_VERIFY_MEMBER(FStructSchema, Footer);
+
+PP_DBGVIS_VERIFY_TYPE(FEnumSchema);
+PP_DBGVIS_VERIFY_MEMBER(FEnumSchema, Type);
+PP_DBGVIS_VERIFY_BITFIELD(FEnumSchema, FlagMode);
+PP_DBGVIS_VERIFY_BITFIELD(FEnumSchema, ExplicitConstants);
+PP_DBGVIS_VERIFY_MEMBER(FEnumSchema, Width);
+PP_DBGVIS_VERIFY_MEMBER(FEnumSchema, Num);
+PP_DBGVIS_VERIFY_MEMBER(FEnumSchema, Footer);
+
+////////////////////////////////////////////////////////////////////////////////
+
+namespace DbgVis
+{
+
+// D4B455B7-7BAB-4F1D-A944-98EC086FB4AB
+static constexpr FGuid GVisualizerGuid = FGuid(0xD4B455B7, 0x7BAB4F1D, 0xA94498EC, 0x086FB4AB);
+
+void FIdScope::AssignDebuggingState()
+{
+	FIdScope::Global = &Current;
+	(void)::UE::Core::FVisualizerDebuggingState::Assign(GVisualizerGuid, Global);
+}
+
+FIdScope::~FIdScope()
+{
+	FIdScope::Global = Previous;
+	(void)::UE::Core::FVisualizerDebuggingState::Assign(GVisualizerGuid, Global);
+}
+
+} // namespace DbgVis
+
+} // namespace PlainProps
diff --git a/Engine/Plugins/Experimental/PlainProps/Source/Public/PlainPropsVisualize.h b/Engine/Plugins/Experimental/PlainProps/Source/Public/PlainPropsVisualize.h
new file mode 100644
--- /dev/null
+++ b/Engine/Plugins/Experimental/PlainProps/Source/Public/PlainPropsVisualize.h
@@ -0,0 +1,106 @@
+// Copyright Epic Games, Inc. All Rights Reserved.
+
+#pragma once
+
+#include "PlainPropsBuild.h"
+#include "PlainPropsIndex.h"
+#include "Containers/Set.h"
+
+namespace PlainProps
+{
+
+// Mirror publically exposed forward-declared internal structs for full visualizer support in the context of other modules.
+namespace DbgVis
+{
+
+// InternalBuild.h
+
+struct FBuiltStruct
+{
+	uint16 NumMembers;
+	FBuiltMember Members[0];
+};
+
+struct FBuiltRange
+{
+	uint64 Num;
+	uint8 Data[0];
+};
+
+
+// InternalFormat.h
+
+struct FSchemaBatch
+{
+	uint32 NumNestedScopes;
+	uint32 NestedScopesOffset;
+	uint32 NumParametricTypes;
+	uint32 NumSchemas;
+	uint32 NumStructSchemas;
+	uint32 SchemaOffsets[0];
+};
+
+enum class ESuper : uint8 { No, Unused, Used, Reused };
+
+struct FStructSchema
+{
+	FType Type;
+	uint16 NumMembers;
+	uint16 NumRangeTypes;
+	uint16 NumInnerSchemas;
+	ESuper Inheritance : 2;
+	uint8 IsDense : 1;
+	FMemberType Footer[0];
+};
+
+struct FEnumSchema
+{
+	FType Type;
+	uint8 FlagMode : 1;
+	uint8 ExplicitConstants : 1;
+	ELeafWidth Width;
+	uint16 Num;
+	FNameId Footer[0];
+};
+
+////////////////////////////////////////////////////////////////////////////////
+
+struct FIdVisualizer
+{
+	const FIdIndexerBase* Indexer;
+	const char* NameType;
+
+	static void KeepDebugInfo(DbgVis::FBuiltStruct*) {};
+	static void KeepDebugInfo(DbgVis::FBuiltRange*) {};
+	static void KeepDebugInfo(DbgVis::ESuper) {};
+	static void KeepDebugInfo(DbgVis::FSchemaBatch*) {};
+	static void KeepDebugInfo(DbgVis::FStructSchema*) {};
+	static void KeepDebugInfo(DbgVis::FEnumSchema*) {};
+};
+
+class FIdScope
+{
+public:
+	// NameType must be a unique static string identifer for NameT,
+	// it is used by natvis in a memcmp(str1,str2,N) expression to select the correct typed name indexer, N <= 8
+	template<typename NameT>
+	FIdScope(const TIdIndexer<NameT>& Indexer, const char NameType[8] = nullptr)
+	: Current{ static_cast<const FIdIndexerBase*>(&Indexer), &NameType[0] }
+	, Previous(Global)
+	{
+		AssignDebuggingState();
+	}
+
+	PLAINPROPS_API ~FIdScope();
+
+private:
+	PLAINPROPS_API void AssignDebuggingState();
+
+	inline static FIdVisualizer*	Global = nullptr;
+	FIdVisualizer					Current;
+	FIdVisualizer*					Previous;
+};
+
+} // namespace DbgVis
+
+} // namespace PlainProps
diff --git a/Engine/Plugins/Experimental/PlainProps/Source/Tests/PlainPropsTests.cpp b/Engine/Plugins/Experimental/PlainProps/Source/Tests/PlainPropsTests.cpp
--- a/Engine/Plugins/Experimental/PlainProps/Source/Tests/PlainPropsTests.cpp
+++ b/Engine/Plugins/Experimental/PlainProps/Source/Tests/PlainPropsTests.cpp
@@ -15,6 +15,7 @@
 #include "PlainPropsWrite.h"
 #include "PlainPropsBuildSchema.h"
 #include "PlainPropsIndex.h"
+#include "PlainPropsVisualize.h"
 #include "PlainPropsInternalBuild.h"
 #include "PlainPropsInternalFormat.h"
 #include "PlainPropsInternalParse.h"
@@ -439,6 +440,7 @@ static void TestSerialize(void (*BuildObjects)(FTestBatchBuilder&, FScratchAlloc
 	{
 		FScratchAllocator Scratch;
 		FTestBatchBuilder Batch(Scratch);
+		DbgVis::FIdScope _(Batch, "AnsiStr");
 		BuildObjects(Batch, Scratch);
 		Data = Batch.Write();
 	}
diff --git a/Engine/Plugins/Experimental/PlainProps/Source/Tests/PlainPropsUeCoreTests.cpp b/Engine/Plugins/Experimental/PlainProps/Source/Tests/PlainPropsUeCoreTests.cpp
--- a/Engine/Plugins/Experimental/PlainProps/Source/Tests/PlainPropsUeCoreTests.cpp
+++ b/Engine/Plugins/Experimental/PlainProps/Source/Tests/PlainPropsUeCoreTests.cpp
@@ -16,6 +16,7 @@
 #include "PlainPropsLoad.h"
 #include "PlainPropsRead.h"
 #include "PlainPropsSave.h"
+#include "PlainPropsVisualize.h"
 #include "PlainPropsWrite.h"
 #include "PlainPropsUeCoreBindings.h"
 #include "Algo/Compare.h"
@@ -938,6 +939,8 @@ TArray<T, TInlineAllocator<N>> MakeInlArray(const T* Str)
 
 TEST_CASE_NAMED(FPlainPropsUeCoreTest, "System::Core::Serialization::PlainProps::UE::Core", "[Core][PlainProps][SmokeFilter]")
 {
+	DbgVis::FIdScope _(GNames, "FName");
+
 	SECTION("Basic")
 	{
 		TScopedDefaultStructBinding<FInt> Int;
diff --git a/Engine/Plugins/Experimental/PlainPropsUObject/Source/Private/PlainPropsUObjectBindings.cpp b/Engine/Plugins/Experimental/PlainPropsUObject/Source/Private/PlainPropsUObjectBindings.cpp
--- a/Engine/Plugins/Experimental/PlainPropsUObject/Source/Private/PlainPropsUObjectBindings.cpp
+++ b/Engine/Plugins/Experimental/PlainPropsUObject/Source/Private/PlainPropsUObjectBindings.cpp
@@ -7,6 +7,7 @@
 #include "PlainPropsParse.h"
 #include "PlainPropsPrint.h"
 #include "PlainPropsRead.h"
+#include "PlainPropsVisualize.h"
 #include "PlainPropsWrite.h"
 #include "Algo/Compare.h"
 #include "Algo/Find.h"
@@ -4385,6 +4386,10 @@ UTestPlainPropsCommandlet::UTestPlainPropsCommandlet(const FObjectInitializer& I
 
 int32 UTestPlainPropsCommandlet::Main(const FString& Params)
 {
+	using namespace PlainProps;
+	using namespace PlainProps::UE;
+	DbgVis::FIdScope _(GUE.Names, "FName");
+
 	if (int32 LoadIdx = Params.Find("-load="); LoadIdx != INDEX_NONE)
 	{
 		// E.g. -run=TestPlainProps -load=/BRRoot/BRRoot.BRRoot,/Game/Maps/FrontEnd.FrontEnd
@@ -4424,7 +4429,6 @@ int32 UTestPlainPropsCommandlet::Main(const FString& Params)
 
 	UE_LOGFMT(LogPlainPropsUObject, Display, "Starting test...");
 
-	using namespace PlainProps::UE;
 	ERoundtrip Options = ERoundtrip::PP | ERoundtrip::UPS | ERoundtrip::TPS | ERoundtrip::TextMemory;
 	if (Params.Find("-pp") != INDEX_NONE)
 	{
@@ -4441,13 +4445,3 @@ int32 UTestPlainPropsCommandlet::Main(const FString& Params)
 	return TestBindings(Options);
 }
 
-// Stopgap debugger callable functions til natvis is added 
-extern FString _print(const PlainProps::FNameId& Id)				{ return PlainProps::UE::GUE.Debug.Print(Id); }
-extern FString _print(const PlainProps::FMemberId& Id)				{ return PlainProps::UE::GUE.Debug.Print(Id); }
-extern FString _print(const PlainProps::FOptionalMemberId& Id)		{ return PlainProps::UE::GUE.Debug.Print(Id); }
-extern FString _print(const PlainProps::FScopeId& Id)				{ return PlainProps::UE::GUE.Debug.Print(Id); }
-extern FString _print(const PlainProps::FTypenameId& Id)			{ return PlainProps::UE::GUE.Debug.Print(Id); }
-extern FString _print(const PlainProps::FConcreteTypenameId& Id)	{ return PlainProps::UE::GUE.Debug.Print(Id); }
-extern FString _print(const PlainProps::FType& Id)					{ return PlainProps::UE::GUE.Debug.Print(Id); }
-extern FString _print(const PlainProps::FEnumId& Id)				{ return PlainProps::UE::GUE.Debug.Print(Id); }
-extern FString _print(const PlainProps::FStructId& Id)				{ return PlainProps::UE::GUE.Debug.Print(Id); }
