PlainProps fix static reflection for structs with 10 members

#rnx

diff --git a/Engine/Plugins/Experimental/PlainProps/Source/Public/PlainPropsCtti.h b/Engine/Plugins/Experimental/PlainProps/Source/Public/PlainPropsCtti.h
--- a/Engine/Plugins/Experimental/PlainProps/Source/Public/PlainPropsCtti.h
+++ b/Engine/Plugins/Experimental/PlainProps/Source/Public/PlainPropsCtti.h
@@ -18,15 +18,16 @@ template<typename T> using CttiOf = typename TCttiOf<T>::Type;
 #define _PP_NUM_ARGS_(_1,_2,_3,_4,_5,_6,_7,_8,_9,_10, N, ...) N
 #define _PP_NUM_ARGS(...) _PP_EXPAND(_PP_NUM_ARGS_(__VA_ARGS__,10,9,8,7,6,5,4,3,2,1))
 
-#define _PP_REFLECT_1(N, MACRO, NS, T, M)				MACRO(N-1, NS, T, M)
-#define _PP_REFLECT_2(N, MACRO, NS, T, M, ...)			MACRO(N-2, NS, T, M) _PP_REFLECT_1(N, MACRO, NS, T, __VA_ARGS__)
-#define _PP_REFLECT_3(N, MACRO, NS, T, M, ...)			MACRO(N-3, NS, T, M) _PP_REFLECT_2(N, MACRO, NS, T, __VA_ARGS__)
-#define _PP_REFLECT_4(N, MACRO, NS, T, M, ...)			MACRO(N-4, NS, T, M) _PP_REFLECT_3(N, MACRO, NS, T, __VA_ARGS__)
-#define _PP_REFLECT_5(N, MACRO, NS, T, M, ...)			MACRO(N-5, NS, T, M) _PP_REFLECT_4(N, MACRO, NS, T, __VA_ARGS__)
-#define _PP_REFLECT_6(N, MACRO, NS, T, M, ...)			MACRO(N-6, NS, T, M) _PP_REFLECT_5(N, MACRO, NS, T, __VA_ARGS__)
-#define _PP_REFLECT_7(N, MACRO, NS, T, M, ...)			MACRO(N-7, NS, T, M) _PP_REFLECT_6(N, MACRO, NS, T, __VA_ARGS__)
-#define _PP_REFLECT_8(N, MACRO, NS, T, M, ...)			MACRO(N-8, NS, T, M) _PP_REFLECT_7(N, MACRO, NS, T, __VA_ARGS__)
-#define _PP_REFLECT_9(N, MACRO, NS, T, M, ...)			MACRO(N-9, NS, T, M) _PP_REFLECT_8(N, MACRO, NS, T, __VA_ARGS__)
+#define _PP_REFLECT_1( N, MACRO, NS, T, M)				MACRO(N-1, NS, T, M)
+#define _PP_REFLECT_2( N, MACRO, NS, T, M, ...)			MACRO(N-2, NS, T, M)  _PP_REFLECT_1(N, MACRO, NS, T, __VA_ARGS__)
+#define _PP_REFLECT_3( N, MACRO, NS, T, M, ...)			MACRO(N-3, NS, T, M)  _PP_REFLECT_2(N, MACRO, NS, T, __VA_ARGS__)
+#define _PP_REFLECT_4( N, MACRO, NS, T, M, ...)			MACRO(N-4, NS, T, M)  _PP_REFLECT_3(N, MACRO, NS, T, __VA_ARGS__)
+#define _PP_REFLECT_5( N, MACRO, NS, T, M, ...)			MACRO(N-5, NS, T, M)  _PP_REFLECT_4(N, MACRO, NS, T, __VA_ARGS__)
+#define _PP_REFLECT_6( N, MACRO, NS, T, M, ...)			MACRO(N-6, NS, T, M)  _PP_REFLECT_5(N, MACRO, NS, T, __VA_ARGS__)
+#define _PP_REFLECT_7( N, MACRO, NS, T, M, ...)			MACRO(N-7, NS, T, M)  _PP_REFLECT_6(N, MACRO, NS, T, __VA_ARGS__)
+#define _PP_REFLECT_8( N, MACRO, NS, T, M, ...)			MACRO(N-8, NS, T, M)  _PP_REFLECT_7(N, MACRO, NS, T, __VA_ARGS__)
+#define _PP_REFLECT_9( N, MACRO, NS, T, M, ...)			MACRO(N-9, NS, T, M)  _PP_REFLECT_8(N, MACRO, NS, T, __VA_ARGS__)
+#define _PP_REFLECT_10(N, MACRO, NS, T, M, ...)			MACRO(N-10, NS, T, M) _PP_REFLECT_9(N, MACRO, NS, T, __VA_ARGS__)
 #define _PP_REFLECT_THINGS(N, MACRO, ...)				_PP_EXPAND(_PP_REFLECT_##N(N, MACRO, __VA_ARGS__))
 
 #define _PP_REFLECT_ENUMERATOR(N, NS, T, C)				{ #C, Type::C }, 
