PlainProps add unused declaration and schema version

#rnx

diff --git a/Engine/Plugins/Experimental/PlainProps/Source/Private/PlainPropsDeclare.cpp b/Engine/Plugins/Experimental/PlainProps/Source/Private/PlainPropsDeclare.cpp
--- a/Engine/Plugins/Experimental/PlainProps/Source/Private/PlainPropsDeclare.cpp
+++ b/Engine/Plugins/Experimental/PlainProps/Source/Private/PlainPropsDeclare.cpp
@@ -97,7 +97,7 @@ void CopyItems(T* It, TConstArrayView<T> Items)
 	}
 }
 
-static const FStructDeclaration& Declare(TArray<TUniquePtr<FStructDeclaration>>& Out, FStructId DeclId, FType Type, TConstArrayView<FMemberId> MemberOrder, EMemberPresence Occupancy, FOptionalStructId Super, uint16 NumMembers)
+static const FStructDeclaration& Declare(TArray<TUniquePtr<FStructDeclaration>>& Out, FStructId DeclId, FType Type, uint16 Version, TConstArrayView<FMemberId> MemberOrder, EMemberPresence Occupancy, FOptionalStructId Super, uint16 NumMembers)
 {
 	if (static_cast<int32>(DeclId.Idx) >= Out.Num())
 	{
@@ -110,15 +110,16 @@ static const FStructDeclaration& Declare(TArray<TUniquePtr<FStructDeclaration>>&
 		check(DeclId == Ptr->Id);
 		check(Type == Ptr->Type);
 		check(Super == Ptr->Super);
-		check(Occupancy == Ptr->Occupancy);
+		check(Version == Ptr->Version);
 		check(NumMembers == Ptr->NumMembers);
+		check(Occupancy == Ptr->Occupancy);
 		check(Algo::Compare(MemberOrder, Ptr->GetMemberOrder()));
 
 		Ptr->RefCount++;
 	}
 	else
 	{
-		FStructDeclaration Header{1, DeclId, Type, Super, Occupancy, NumMembers};
+		FStructDeclaration Header{1, DeclId, Type, Super, Version, NumMembers, Occupancy};
 		void* Data = FMemory::Malloc(sizeof(FStructDeclaration) + MemberOrder.Num() * MemberOrder.GetTypeSize());
 		Ptr.Reset(new (Data) FStructDeclaration(Header));
 		CopyItems(Ptr->MemberOrder, MemberOrder);
@@ -127,15 +128,15 @@ static const FStructDeclaration& Declare(TArray<TUniquePtr<FStructDeclaration>>&
 	return *Ptr;
 }
 
-const FStructDeclaration& FDeclarations::DeclareStruct(FStructId DeclId, FType Type, TConstArrayView<FMemberId> MemberOrder, EMemberPresence Occupancy, FOptionalStructId Super)
+const FStructDeclaration& FDeclarations::DeclareStruct(FStructId DeclId, FType Type, uint16 Version, TConstArrayView<FMemberId> MemberOrder, EMemberPresence Occupancy, FOptionalStructId Super)
 {
 	checkf(!(Super && Occupancy == EMemberPresence::RequireAll), TEXT("'%s' is a dense substruct, this isn't supported, see BuildSuperStruct()"), *Debug.Print(DeclId));
-	return Declare(/* out */ DeclaredStructs, DeclId, Type, MemberOrder, Occupancy, Super, IntCastChecked<uint16>(MemberOrder.Num()));
+	return Declare(/* out */ DeclaredStructs, DeclId, Type, Version, MemberOrder, Occupancy, Super, IntCastChecked<uint16>(MemberOrder.Num()));
 }
 
 const FStructDeclaration& FDeclarations::DeclareNumeralStruct(FStructId DeclId, FType Type, TConstArrayView<FMemberId> Numerals, EMemberPresence Occupancy)
 {
-	return Declare(/* out */ DeclaredStructs, DeclId, Type, Numerals, Occupancy, NoId, IntCastChecked<uint16>(Numerals.Num()));
+	return Declare(/* out */ DeclaredStructs, DeclId, Type, 0, Numerals, Occupancy, NoId, IntCastChecked<uint16>(Numerals.Num()));
 }
 
 const FEnumDeclaration& FDeclarations::DeclareEnum(FEnumId Id, FType Type, EEnumMode Mode, TConstArrayView<FEnumerator> Enumerators, EEnumAliases Policy)
diff --git a/Engine/Plugins/Experimental/PlainProps/Source/Private/PlainPropsInternalFormat.h b/Engine/Plugins/Experimental/PlainProps/Source/Private/PlainPropsInternalFormat.h
--- a/Engine/Plugins/Experimental/PlainProps/Source/Private/PlainPropsInternalFormat.h
+++ b/Engine/Plugins/Experimental/PlainProps/Source/Private/PlainPropsInternalFormat.h
@@ -64,6 +64,7 @@ inline bool SkipDeclaredSuperSchema(ESuper Inheritance)	{ return Inheritance ==
 struct FStructSchema
 {
 	FType Type;
+	uint16 Version;
 	uint16 NumMembers;
 	uint16 NumRangeTypes;
 	uint16 NumInnerSchemas;
@@ -100,14 +101,15 @@ struct FStructSchema
 	}
 };
 
-static_assert(sizeof(FStructSchema) == 16 && alignof(FStructSchema) == 4,		"Add binary format versioning and read old format in AllocateReadSchemas");
+static_assert(sizeof(FStructSchema) == 20 && alignof(FStructSchema) == 4,		"Add binary format versioning and read old format in AllocateReadSchemas");
 static_assert(offsetof(FStructSchema, Type) == 0,								"Add binary format versioning and read old format in AllocateReadSchemas");
-static_assert(offsetof(FType, Scope) == 0,									"Add binary format versioning and read old format in AllocateReadSchemas");
+static_assert(offsetof(FType, Scope) == 0,										"Add binary format versioning and read old format in AllocateReadSchemas");
 static_assert(offsetof(FType, Name) == 4,										"Add binary format versioning and read old format in AllocateReadSchemas");
-static_assert(offsetof(FStructSchema, NumMembers) == 8,							"Add binary format versioning and read old format in AllocateReadSchemas");
-static_assert(offsetof(FStructSchema, NumRangeTypes) == 10,						"Add binary format versioning and read old format in AllocateReadSchemas");
-static_assert(offsetof(FStructSchema, NumInnerSchemas) == 12,					"Add binary format versioning and read old format in AllocateReadSchemas");
-static_assert(offsetof(FStructSchema, Footer) == 15,							"Add binary format versioning and read old format in AllocateReadSchemas");
+static_assert(offsetof(FStructSchema, Version) == 8,							"Add binary format versioning and read old format in AllocateReadSchemas");
+static_assert(offsetof(FStructSchema, NumMembers) == 10,						"Add binary format versioning and read old format in AllocateReadSchemas");
+static_assert(offsetof(FStructSchema, NumRangeTypes) == 12,						"Add binary format versioning and read old format in AllocateReadSchemas");
+static_assert(offsetof(FStructSchema, NumInnerSchemas) == 14,					"Add binary format versioning and read old format in AllocateReadSchemas");
+static_assert(offsetof(FStructSchema, Footer) == 17,							"Add binary format versioning and read old format in AllocateReadSchemas");
 
 inline uint32 CalculateSize(const FStructSchema& In)
 {
diff --git a/Engine/Plugins/Experimental/PlainProps/Source/Private/PlainPropsInternalParse.h b/Engine/Plugins/Experimental/PlainProps/Source/Private/PlainPropsInternalParse.h
--- a/Engine/Plugins/Experimental/PlainProps/Source/Private/PlainPropsInternalParse.h
+++ b/Engine/Plugins/Experimental/PlainProps/Source/Private/PlainPropsInternalParse.h
@@ -72,6 +72,7 @@ struct FParsedMemberSchema
 struct FParsedStructSchema
 {
 	FStructId						Id;
+	uint16							Version = 0;
 	FOptionalStructId				DeclaredSuper;
 	FOptionalStructId				UsedSuper;
 	TConstArrayView<FMemberId>		MemberNames;	// FStructDeclaration.MemberOrder
diff --git a/Engine/Plugins/Experimental/PlainProps/Source/Private/PlainPropsInternalText.h b/Engine/Plugins/Experimental/PlainProps/Source/Private/PlainPropsInternalText.h
--- a/Engine/Plugins/Experimental/PlainProps/Source/Private/PlainPropsInternalText.h
+++ b/Engine/Plugins/Experimental/PlainProps/Source/Private/PlainPropsInternalText.h
@@ -14,6 +14,7 @@ struct FLiterals
 	FAnsiStringView Enums = "Enums";
 	FAnsiStringView Objects = "Objects";
 
+	FAnsiStringView Version = "Version";
 	FAnsiStringView DeclaredSuper = "Super";
 	FAnsiStringView Width = "Width";
 	FAnsiStringView FlagMode = "FlagMode";
diff --git a/Engine/Plugins/Experimental/PlainProps/Source/Private/PlainPropsParse.cpp b/Engine/Plugins/Experimental/PlainProps/Source/Private/PlainPropsParse.cpp
--- a/Engine/Plugins/Experimental/PlainProps/Source/Private/PlainPropsParse.cpp
+++ b/Engine/Plugins/Experimental/PlainProps/Source/Private/PlainPropsParse.cpp
@@ -598,6 +598,7 @@ private:
 
 	FParsedStructSchema DeclareStructSchema(
 		const FTypeTokens& TypeTokens,
+		uint16 Version,
 		const TOptional<FTypeTokens>& ParsedSuper,
 		TConstArrayView<FUtf8StringView> MemberNames,
 		TConstArrayView<FUtf8StringView> MemberTypes);
@@ -1108,7 +1109,12 @@ bool FStructSchemaParser::ParseStructSchema(FParsedStructSchema& Out)
 		Token = ScopedReader.GrabToken();
 		if (Token.Token == EToken::LeafId)
 		{
-			if (Token.Value() == GLiterals.DeclaredSuper)
+			if (Token.Value() == GLiterals.Version)
+			{
+				Token = ScopedReader.GrabToken();
+				Out.Version = *Parse<uint16>(Token.Value());
+			}
+			else if (Token.Value() == GLiterals.DeclaredSuper)
 			{
 				Token = ScopedReader.GrabToken();
 				ParsedSuper = TokenizeType(Token.Value());
@@ -1151,12 +1157,13 @@ bool FStructSchemaParser::ParseStructSchema(FParsedStructSchema& Out)
 		}
 	}
 
-	Out = DeclareStructSchema(TypeTokens.GetValue(), ParsedSuper, MemberNames, MemberTypes);
+	Out = DeclareStructSchema(TypeTokens.GetValue(), Out.Version, ParsedSuper, MemberNames, MemberTypes);
 	return true;
 }
 
 FParsedStructSchema FStructSchemaParser::DeclareStructSchema(
 	const FTypeTokens& TypeTokens,
+	uint16 Version,
 	const TOptional<FTypeTokens>& ParsedSuper,
 	TConstArrayView<FUtf8StringView> MemberNames,
 	TConstArrayView<FUtf8StringView> MemberTypes)
@@ -1198,9 +1205,9 @@ FParsedStructSchema FStructSchemaParser::DeclareStructSchema(
 		}
 	}
 
-	Types.DeclareStruct(Id, Type, MemberIds, EMemberPresence::AllowSparse, DeclaredSuper);
+	Types.DeclareStruct(Id, Type, Version, MemberIds, EMemberPresence::AllowSparse, DeclaredSuper);
 
-	return { Id, DeclaredSuper, UsedSuper, Types.Get(Id).GetMemberOrder(), MoveTemp(MemberSchemas) };
+	return { Id, Version, DeclaredSuper, UsedSuper, Types.Get(Id).GetMemberOrder(), MoveTemp(MemberSchemas) };
 }
 
 static TOptional<ERangeSizeType> GrabRangeSize(FUtf8StringView& String)
diff --git a/Engine/Plugins/Experimental/PlainProps/Source/Private/PlainPropsPrint.cpp b/Engine/Plugins/Experimental/PlainProps/Source/Private/PlainPropsPrint.cpp
--- a/Engine/Plugins/Experimental/PlainProps/Source/Private/PlainPropsPrint.cpp
+++ b/Engine/Plugins/Experimental/PlainProps/Source/Private/PlainPropsPrint.cpp
@@ -280,6 +280,7 @@ public:
 	FType					GetStruct() const		{ return Struct; }
 	bool					IsDense() const			{ return bIsDense; }
 	bool					HasSuper() const		{ return bHasSuper; }
+	uint16					GetVersion() const		{ return Version; }
 
 	bool					HasMore() const			{ return MemberIdx < NumMembers; }
 	
@@ -297,6 +298,7 @@ private:
 	const bool				bIsDense : 1;
 	const bool				bHasSuper : 1;
 	const bool				bUsesSuper : 1;
+	const uint16			Version;
 	const uint32			NumMembers;
 	const uint32			NumRangeTypes;			// Number of ranges and nested ranges
 	const uint32			NumInnerSchemas;		// Number of static structs and enums
@@ -436,6 +438,7 @@ FStructSchemaReader::FStructSchemaReader(const FStructSchema& Schema, FReadBatch
 , bIsDense(Schema.IsDense)
 , bHasSuper(Schema.Inheritance != ESuper::No)
 , bUsesSuper(UsesSuper(Schema.Inheritance))
+, Version(Schema.Version)
 , NumMembers(Schema.NumMembers)
 , NumRangeTypes(Schema.NumRangeTypes)
 , NumInnerSchemas(Schema.NumInnerSchemas)
@@ -810,6 +813,11 @@ void FBatchPrinter::PrintStructSchema(const FStructSchema& Struct, FReadBatchId
 
 	TextBuilder.BeginStruct(*FPrintId<128>(Ids, Reader.GetStruct()));
 
+	if (uint16 Version = Reader.GetVersion())
+	{
+		TextBuilder.AddLeaf(GLiterals.Version, Version);
+	}
+
 	if (Reader.HasSuper())
 	{
 		const FStructSchema& SuperSchema = Reader.GetSuper().Resolve();
diff --git a/Engine/Plugins/Experimental/PlainProps/Source/Private/PlainPropsVisualize.cpp b/Engine/Plugins/Experimental/PlainProps/Source/Private/PlainPropsVisualize.cpp
--- a/Engine/Plugins/Experimental/PlainProps/Source/Private/PlainPropsVisualize.cpp
+++ b/Engine/Plugins/Experimental/PlainProps/Source/Private/PlainPropsVisualize.cpp
@@ -58,6 +58,7 @@ PP_DBGVIS_VERIFY_MEMBER(FStructSchema, Type);
 PP_DBGVIS_VERIFY_MEMBER(FStructSchema, NumMembers);
 PP_DBGVIS_VERIFY_MEMBER(FStructSchema, NumRangeTypes);
 PP_DBGVIS_VERIFY_MEMBER(FStructSchema, NumInnerSchemas);
+PP_DBGVIS_VERIFY_MEMBER(FStructSchema, Version);
 PP_DBGVIS_VERIFY_BITFIELD(FStructSchema, Inheritance);
 PP_DBGVIS_VERIFY_BITFIELD(FStructSchema, IsDense);
 PP_DBGVIS_VERIFY_MEMBER(FStructSchema, Footer);
diff --git a/Engine/Plugins/Experimental/PlainProps/Source/Public/PlainPropsBind.h b/Engine/Plugins/Experimental/PlainProps/Source/Public/PlainPropsBind.h
--- a/Engine/Plugins/Experimental/PlainProps/Source/Public/PlainPropsBind.h
+++ b/Engine/Plugins/Experimental/PlainProps/Source/Public/PlainPropsBind.h
@@ -873,7 +873,7 @@ struct TScopedStructBinding : FDualStructId, CustomBinding
 	: FDualStructId(DualId)
 	, CustomBinding(TCustomInit<Ids>{{Lowered}})
 	{
-		const FStructDeclaration& Declaration = Runtime::GetTypes().DeclareStruct(DeclId, DeclName, CustomBinding::MemberIds, TOccupancyOf<T>::Value);
+		const FStructDeclaration& Declaration = Runtime::GetTypes().DeclareStruct(DeclId, DeclName, 0, CustomBinding::MemberIds, TOccupancyOf<T>::Value);
 		Runtime::GetCustoms().BindStruct(BindId, *this, Declaration, Lowered);
 	}
 
@@ -1115,7 +1115,7 @@ FStructId DeclareNativeStruct(FDeclarations& Out, EMemberPresence Occupancy)
 	{ 
 		MemberIds[Var::Index] = Ids::IndexMember(Var::Name);
 	});
-	Out.DeclareStruct(Id, Type, MemberIds, Occupancy, SuperId);
+	Out.DeclareStruct(Id, Type, 0, MemberIds, Occupancy, SuperId);
 
 	return Id;
 }
diff --git a/Engine/Plugins/Experimental/PlainProps/Source/Public/PlainPropsDeclare.h b/Engine/Plugins/Experimental/PlainProps/Source/Public/PlainPropsDeclare.h
--- a/Engine/Plugins/Experimental/PlainProps/Source/Public/PlainPropsDeclare.h
+++ b/Engine/Plugins/Experimental/PlainProps/Source/Public/PlainPropsDeclare.h
@@ -40,8 +40,9 @@ struct FStructDeclaration
 	FStructId				Id;				// Could be removed, might allow declaration dedup among templated types
 	FType					Type;			// Could be removed, might allow declaration dedup among templated types
 	FOptionalStructId		Super;
-	EMemberPresence			Occupancy;
+	uint16					Version;
 	uint16					NumMembers;
+	EMemberPresence			Occupancy;
 	FMemberId				MemberOrder[0];
 
 	TConstArrayView<FMemberId> GetMemberOrder() const
@@ -64,7 +65,7 @@ public:
 
 	PLAINPROPS_API const FEnumDeclaration&			DeclareEnum(FEnumId DeclId, FType Type, EEnumMode Mode, TConstArrayView<FEnumerator> Enumerators, EEnumAliases Policy);
 	// Declare struct with ref count 1 or increment it and check that previous declaration matches
-	PLAINPROPS_API const FStructDeclaration&		DeclareStruct(FStructId DeclId, FType Type, TConstArrayView<FMemberId> MemberOrder, EMemberPresence Occupancy, FOptionalStructId Super = {});
+	PLAINPROPS_API const FStructDeclaration&		DeclareStruct(FStructId DeclId, FType Type, uint16 Version, TConstArrayView<FMemberId> MemberOrder, EMemberPresence Occupancy, FOptionalStructId Super = {});
 	PLAINPROPS_API const FStructDeclaration&		DeclareNumeralStruct(FStructId DeclId, FType Type, TConstArrayView<FMemberId> Numerals, EMemberPresence Occupancy);
 	
 	void											DropEnum(FEnumId DeclId)		{ Check(DeclId); DeclaredEnums[DeclId.Idx].Reset(); }
diff --git a/Engine/Plugins/Experimental/PlainProps/Source/Public/PlainPropsVisualize.h b/Engine/Plugins/Experimental/PlainProps/Source/Public/PlainPropsVisualize.h
--- a/Engine/Plugins/Experimental/PlainProps/Source/Public/PlainPropsVisualize.h
+++ b/Engine/Plugins/Experimental/PlainProps/Source/Public/PlainPropsVisualize.h
@@ -45,6 +45,7 @@ enum class ESuper : uint8 { No, Unused, Used, Reused };
 struct FStructSchema
 {
 	FType Type;
+	uint16 Version;
 	uint16 NumMembers;
 	uint16 NumRangeTypes;
 	uint16 NumInnerSchemas;
diff --git a/Engine/Plugins/Experimental/PlainProps/Source/Tests/PlainPropsTests.cpp b/Engine/Plugins/Experimental/PlainProps/Source/Tests/PlainPropsTests.cpp
--- a/Engine/Plugins/Experimental/PlainProps/Source/Tests/PlainPropsTests.cpp
+++ b/Engine/Plugins/Experimental/PlainProps/Source/Tests/PlainPropsTests.cpp
@@ -182,7 +182,7 @@ private:
 FStructId FTestBatchBuilder::DeclareStruct(FType Type, std::initializer_list<const char*> MemberOrder, EMemberPresence Occupancy, FOptionalStructId Super)
 {
 	FStructId Id = IndexStruct(Type);
-	Declarations.DeclareStruct(Id, Type, NameMembers(MemberOrder), Occupancy, Super);
+	Declarations.DeclareStruct(Id, Type, 0, NameMembers(MemberOrder), Occupancy, Super);
 	return Id;
 }
 
diff --git a/Engine/Plugins/Experimental/PlainProps/Source/Tests/PlainPropsUeCoreTests.cpp b/Engine/Plugins/Experimental/PlainProps/Source/Tests/PlainPropsUeCoreTests.cpp
--- a/Engine/Plugins/Experimental/PlainProps/Source/Tests/PlainPropsUeCoreTests.cpp
+++ b/Engine/Plugins/Experimental/PlainProps/Source/Tests/PlainPropsUeCoreTests.cpp
@@ -100,7 +100,7 @@ struct FNameDeclaration
 	: Id(FRuntimeIds::IndexStruct(Type))
 	, Idx(FRuntimeIds::IndexMember("Idx"))
 	{
-		GTypes.DeclareStruct(Id, Type, MakeArrayView(&Idx, 1), EMemberPresence::RequireAll);
+		GTypes.DeclareStruct(Id, Type, 123, MakeArrayView(&Idx, 1), EMemberPresence::RequireAll);
 	}
 
 	~FNameDeclaration()
diff --git a/Engine/Plugins/Experimental/PlainPropsUObject/Source/Private/PlainPropsUObjectBindings.cpp b/Engine/Plugins/Experimental/PlainPropsUObject/Source/Private/PlainPropsUObjectBindings.cpp
--- a/Engine/Plugins/Experimental/PlainPropsUObject/Source/Private/PlainPropsUObjectBindings.cpp
+++ b/Engine/Plugins/Experimental/PlainPropsUObject/Source/Private/PlainPropsUObjectBindings.cpp
@@ -520,7 +520,7 @@ static void DeclareStruct(const UStruct* Struct, FType Type, FStructId Id)
 	}
 	DeclareMembers(/* out */ Members, Struct);
 
-	GUE.Types.DeclareStruct(Id, Type, Members, Occupancy, Super);
+	GUE.Types.DeclareStruct(Id, Type, 0, Members, Occupancy, Super);
 }
 
 static FStructId DeclareStruct(const UStruct* Struct)
@@ -933,7 +933,7 @@ class FInterfaceBindings
 public:
 	FInterfaceBindings()
 	: ScriptInterface(GUE.Names.Resolve(GUE.Structs.ScriptInterface))
-	, Declaration(GUE.Types.DeclareStruct(GUE.Structs.ScriptInterface, ScriptInterface, {GUE.Members.Object}, EMemberPresence::RequireAll))
+	, Declaration(GUE.Types.DeclareStruct(GUE.Structs.ScriptInterface, ScriptInterface, 0, {GUE.Members.Object}, EMemberPresence::RequireAll))
 	{}
 
 	FStructId Bind(FInterfaceProperty* Property)
@@ -2273,7 +2273,7 @@ public:
 		FMemberId Members[] = { GUE.Members.Assign, GUE.Members.Remove, GUE.Members.Insert };
 
 		// Todo: Ownership / memory leak
-		const FStructDeclaration& Declaration = GUE.Types.DeclareStruct(DeclId, DeclType, Members, EMemberPresence::AllowSparse);
+		const FStructDeclaration& Declaration = GUE.Types.DeclareStruct(DeclId, DeclType, 0, Members, EMemberPresence::AllowSparse);
 		DeltaBindNew(Property, Elem, BindId, Declaration);
 		return Bindings.Emplace(Elem, BindId);
 	}
@@ -2608,7 +2608,7 @@ class FMapBindings
 		FMemberId MemberIds[] = { GUE.Members.Assign, GUE.Members.Remove, GUE.Members.Insert };
 
 		// Todo: Ownership / memory leak
-		const FStructDeclaration& Declaration = GUE.Types.DeclareStruct(Id, Type, MemberIds, EMemberPresence::AllowSparse);
+		const FStructDeclaration& Declaration = GUE.Types.DeclareStruct(Id, Type, 0, MemberIds, EMemberPresence::AllowSparse);
 		auto Out = new TMapPropertyCustomBinding<KeyKind, ValueKind>(Property, Members);
 		GUE.Customs.BindStruct(Declaration.Id, *Out, Declaration, {});
 
@@ -2710,7 +2710,7 @@ class FPairBindings
 		FMemberId Members[2] = { GUE.Members.Key, GUE.Members.Value };
 
 		// Todo: Ownership / memory leak
-		GUE.Types.DeclareStruct(DeclId, DeclType, Members, EMemberPresence::RequireAll);
+		GUE.Types.DeclareStruct(DeclId, DeclType, 0, Members, EMemberPresence::RequireAll);
 		GUE.Schemas.BindStruct(BindId, DeclId, Pair.KV);
 
 		Bindings.Emplace(Pair, BindId);
@@ -3396,7 +3396,7 @@ static void BindInitialTypes()
 
 static void DeclareProperty(FStructId Id, TConstArrayView<FMemberId> Members, EMemberPresence Occupancy)
 {
-	GUE.Types.DeclareStruct(Id,	GUE.Names.Resolve(Id), Members, Occupancy);
+	GUE.Types.DeclareStruct(Id,	GUE.Names.Resolve(Id), 0, Members, Occupancy);
 }
 
 static void InitBatchedProperties()
