PlainProps: Add bit field _Pad members to FLeafType, FRangeType and FStructType, FLeafBindType, and FBasicLeafBindType required for zero initialization and deterministic
FMemberType/FMemberBindType::AsByte() conversions.

#rnx

diff --git a/Engine/Plugins/Experimental/PlainProps/Source/Private/PlainPropsRead.cpp b/Engine/Plugins/Experimental/PlainProps/Source/Private/PlainPropsRead.cpp
--- a/Engine/Plugins/Experimental/PlainProps/Source/Private/PlainPropsRead.cpp
+++ b/Engine/Plugins/Experimental/PlainProps/Source/Private/PlainPropsRead.cpp
@@ -10,6 +10,8 @@
 namespace PlainProps
 {
 
+static_assert(sizeof(FMemberType) == 1);
+
 void FSchemaBatch::ValidateBounds(uint64 NumBytes) const
 {
 	static constexpr uint32 Alignment = FMath::Max(alignof(FStructSchema), alignof(FEnumSchema));
diff --git a/Engine/Plugins/Experimental/PlainProps/Source/Public/PlainPropsBind.h b/Engine/Plugins/Experimental/PlainProps/Source/Public/PlainPropsBind.h
--- a/Engine/Plugins/Experimental/PlainProps/Source/Public/PlainPropsBind.h
+++ b/Engine/Plugins/Experimental/PlainProps/Source/Public/PlainPropsBind.h
@@ -57,6 +57,7 @@ struct FBasicLeafBindType
     EMemberKind		_ : 2;
 	ELeafBindType	__ : 3;
 	ELeafWidth		Width : 2;
+	uint8			_Pad : 1;
 };
 
 struct FBitfieldBoolBindType
@@ -78,6 +79,7 @@ union FLeafBindType
 	{
 		EMemberKind			_ : 2;
 		ELeafBindType		Type : 3;
+		uint8				_Pad : 3;
 	}						Bind;
 	FBasicLeafBindType		Basic;
 	FBitfieldBoolBindType	Bitfield;
diff --git a/Engine/Plugins/Experimental/PlainProps/Source/Public/PlainPropsTypes.h b/Engine/Plugins/Experimental/PlainProps/Source/Public/PlainPropsTypes.h
--- a/Engine/Plugins/Experimental/PlainProps/Source/Public/PlainPropsTypes.h
+++ b/Engine/Plugins/Experimental/PlainProps/Source/Public/PlainPropsTypes.h
@@ -57,15 +57,17 @@ template<> inline constexpr ELeafWidth LeafWidth<8> = ELeafWidth::B64;
 
 struct FLeafType
 {
-    EMemberKind		_ : 2;
-    ELeafWidth		Width : 2;
-    ELeafType		Type : 3;
+	EMemberKind		_ : 2;
+	ELeafWidth		Width : 2;
+	ELeafType		Type : 3;
+	uint8			_Pad : 1;
 };
 
 struct FRangeType
 {
-    EMemberKind		_ : 2;
-    ERangeSizeType	MaxSize : 4;
+	EMemberKind		_ : 2;
+	ERangeSizeType	MaxSize : 4;
+	uint8			_Pad : 2;
 };
 
 struct FStructType
@@ -74,6 +76,7 @@ struct FStructType
 	uint8			IsDynamic : 1;	// Which schema stored in value stream, different instances in same batch have different types
 	uint8			IsSuper : 1;	// Inherited members from some super struct, can only appear as first member
 //	uint8			IsDense : 1;	// Potential optimization to avoid accessing schema
+	uint8			_Pad : 4;
 };
 
 union FMemberType
