PlainProps rename FReadBatchId -> FSchemaBatchId and renamed some outdated variable names

#rnx

diff --git a/Engine/Plugins/Experimental/PlainProps/Source/Private/PlainPropsBind.cpp b/Engine/Plugins/Experimental/PlainProps/Source/Private/PlainPropsBind.cpp
--- a/Engine/Plugins/Experimental/PlainProps/Source/Private/PlainPropsBind.cpp
+++ b/Engine/Plugins/Experimental/PlainProps/Source/Private/PlainPropsBind.cpp
@@ -334,17 +334,17 @@ FCustomBindingHandle FCustomBindingsOverlay::Find(FBindId Id) const
 
 ////////////////////////////////////////////////////////////////////////////////////////////////
 
-inline FType IndexInnermost(FIdIndexerBase& Ids, FMemberBindType Type, FOptionalInnerId Schema)
+inline FType IndexInnermost(FIdIndexerBase& Ids, FMemberBindType Type, FOptionalInnerId Id)
 {
 	if (Type.IsStruct())
 	{
-		return Ids.Resolve(Schema.Get().AsStruct());
+		return Ids.Resolve(Id.Get().AsStruct());
 	}
 
 	FUnpackedLeafType Leaf = ToLeafType(Type.AsLeaf());
 	if (Leaf.Type == ELeafType::Enum)
 	{
-		return Ids.Resolve(Schema.Get().AsEnum());
+		return Ids.Resolve(Id.Get().AsEnum());
 	}
 
 	return Ids.MakeLeafParameter(Leaf);
diff --git a/Engine/Plugins/Experimental/PlainProps/Source/Private/PlainPropsBuild.cpp b/Engine/Plugins/Experimental/PlainProps/Source/Private/PlainPropsBuild.cpp
--- a/Engine/Plugins/Experimental/PlainProps/Source/Private/PlainPropsBuild.cpp
+++ b/Engine/Plugins/Experimental/PlainProps/Source/Private/PlainPropsBuild.cpp
@@ -236,13 +236,13 @@ FBuiltMember::FBuiltMember(FMemberId Name, FTypedRange Range)
 : FBuiltMember(Name, MoveTemp(Range.Schema), { .Range = Range.Values })
 {}
 
-FBuiltMember::FBuiltMember(FMemberId Name, FStructId Schema, FBuiltStruct* Value)
-: FBuiltMember(Name, MakeMemberSchema(DefaultStructType, FInnerId(Schema)), { .Struct = Value })
+FBuiltMember::FBuiltMember(FMemberId Name, FStructId Id, FBuiltStruct* Value)
+: FBuiltMember(Name, MakeMemberSchema(DefaultStructType, FInnerId(Id)), { .Struct = Value })
 {}
 
-FBuiltMember FBuiltMember::MakeSuper(FStructId Schema, FBuiltStruct* Value)
+FBuiltMember FBuiltMember::MakeSuper(FStructId Id, FBuiltStruct* Value)
 {
-	return FBuiltMember(NoId, MakeMemberSchema(SuperStructType, FInnerId(Schema)), { .Struct = Value });
+	return FBuiltMember(NoId, MakeMemberSchema(SuperStructType, FInnerId(Id)), { .Struct = Value });
 }
 
 //////////////////////////////////////////////////////////////////////////
diff --git a/Engine/Plugins/Experimental/PlainProps/Source/Private/PlainPropsDiff.cpp b/Engine/Plugins/Experimental/PlainProps/Source/Private/PlainPropsDiff.cpp
--- a/Engine/Plugins/Experimental/PlainProps/Source/Private/PlainPropsDiff.cpp
+++ b/Engine/Plugins/Experimental/PlainProps/Source/Private/PlainPropsDiff.cpp
@@ -273,7 +273,7 @@ static bool DiffSchemas(const FSchemaBatch& A, const FSchemaBatch& B)
 	return Diff;
 }
 
-bool DiffSchemas(FReadBatchId A, FReadBatchId B)
+bool DiffSchemas(FSchemaBatchId A, FSchemaBatchId B)
 {
 	return DiffSchemas(GetReadSchemas(A), GetReadSchemas(B));
 }
diff --git a/Engine/Plugins/Experimental/PlainProps/Source/Private/PlainPropsInternalPrint.h b/Engine/Plugins/Experimental/PlainProps/Source/Private/PlainPropsInternalPrint.h
--- a/Engine/Plugins/Experimental/PlainProps/Source/Private/PlainPropsInternalPrint.h
+++ b/Engine/Plugins/Experimental/PlainProps/Source/Private/PlainPropsInternalPrint.h
@@ -43,7 +43,7 @@ public:
 	void PrintObjects(TConstArrayView<FStructView> Structs);
 
 private:
-	void PrintStructSchema(const FStructSchema& Struct, FReadBatchId BatchId);
+	void PrintStructSchema(const FStructSchema& Struct, FSchemaBatchId BatchId);
 	void PrintEnumSchema(const FEnumSchema& Enum);
 
 	template<typename IntType>
diff --git a/Engine/Plugins/Experimental/PlainProps/Source/Private/PlainPropsInternalRead.h b/Engine/Plugins/Experimental/PlainProps/Source/Private/PlainPropsInternalRead.h
--- a/Engine/Plugins/Experimental/PlainProps/Source/Private/PlainPropsInternalRead.h
+++ b/Engine/Plugins/Experimental/PlainProps/Source/Private/PlainPropsInternalRead.h
@@ -98,7 +98,7 @@ inline TSchemaRange<const FStructSchema>	GetStructSchemas(const FSchemaBatch& Ba
 
 //////////////////////////////////////////////////////////////////////////
 
-const FSchemaBatch& GetReadSchemas(FReadBatchId Batch);
+const FSchemaBatch& GetReadSchemas(FSchemaBatchId Batch);
 
 inline const FStructSchema& ResolveStructSchema(const FSchemaBatch& Batch, FSchemaId Id)
 {
diff --git a/Engine/Plugins/Experimental/PlainProps/Source/Private/PlainPropsInternalTest.h b/Engine/Plugins/Experimental/PlainProps/Source/Private/PlainPropsInternalTest.h
--- a/Engine/Plugins/Experimental/PlainProps/Source/Private/PlainPropsInternalTest.h
+++ b/Engine/Plugins/Experimental/PlainProps/Source/Private/PlainPropsInternalTest.h
@@ -9,9 +9,9 @@
 namespace PlainProps
 {
 
-class FReadBatchId;
+class FSchemaBatchId;
 struct FStructView;
 
-[[nodiscard]] FReadBatchId ParseBatch(TArray64<uint8>& OutData, TArray<FStructView>& OutObjects, FUtf8StringView YamlView);
+[[nodiscard]] FSchemaBatchId ParseBatch(TArray64<uint8>& OutData, TArray<FStructView>& OutObjects, FUtf8StringView YamlView);
 
 } // namespace PlainProps
diff --git a/Engine/Plugins/Experimental/PlainProps/Source/Private/PlainPropsLoad.cpp b/Engine/Plugins/Experimental/PlainProps/Source/Private/PlainPropsLoad.cpp
--- a/Engine/Plugins/Experimental/PlainProps/Source/Private/PlainPropsLoad.cpp
+++ b/Engine/Plugins/Experimental/PlainProps/Source/Private/PlainPropsLoad.cpp
@@ -238,7 +238,7 @@ struct FLoadIdBinderDummy
 
 struct FLoadBatch
 {
-	FReadBatchId			ReadId; // Needed to access schemas for custom struct loading
+	FSchemaBatchId			BatchId; // Needed to access schemas for custom struct loading
 	uint32					NumReadSchemas;
 	uint32					NumPlans;
 	FLoadStructPlan			Plans[0];
@@ -534,7 +534,7 @@ static TOptional<TConstArrayView<T>> GetContiguousSubset(TConstArrayView<T> View
 struct FLoadPlanner
 {
 	// Constructor parameters
-	const FReadBatchId								ReadId;
+	const FSchemaBatchId							BatchId;
 	const FDeclarations&							Declarations;
 	const FCustomBindings&							Customs;
 	const FSchemaBindings&							Schemas;
@@ -548,7 +548,7 @@ struct FLoadPlanner
 	
 	FLoadBatchPtr CreatePlans(ESchemaFormat Format)
 	{
-		check(NumStructSchemas(ReadId) == RuntimeIds.Num());
+		check(NumStructSchemas(BatchId) == RuntimeIds.Num());
 
 		// Make load plans for saved schemas
 		const uint32 NumPlans = RuntimeIds.Num();
@@ -623,7 +623,7 @@ private:
 		SIZE_T Bytes = sizeof(FLoadBatch) + sizeof(FLoadStructPlan) * NumPlans
 					 + Align(sizeof(FStructSchemaId) * NumMappings, alignof(FSchemaBinding))
 					 + SubsetSchemaData.Num();
-		FLoadBatch Header = { ReadId, NumReadSchemas, NumPlans };
+		FLoadBatch Header = { BatchId, NumReadSchemas, NumPlans };
 		FLoadBatch* Out = new (FMemory::Malloc(Bytes)) FLoadBatch{Header};
 		FMemory::Memcpy(Out->Plans, Plans.GetData(), sizeof(FLoadStructPlan) * NumPlans);
 
@@ -739,9 +739,8 @@ private:
 		FInnerStructSchemaIterator SchemaIt(Schema);
 		for (FInnerStruct Inner : Inners)
 		{
-			// Todo: Rename FReadBatchId ReadId everywhere
-			FStructSchemaId ReadId_ = SchemaIt.GrabMemberStruct(Inner.Name);
-			*OutLoadIdIt++ = LoadIds.BindLoadId(ReadId_, Inner.Id);
+			FStructSchemaId ReadId = SchemaIt.GrabMemberStruct(Inner.Name);
+			*OutLoadIdIt++ = LoadIds.BindLoadId(ReadId, Inner.Id);
 		}
 		check((void*)OutLoadIdIt == SubsetSchemaData.GetData() + SubsetSchemaData.Num());
 
@@ -791,7 +790,7 @@ private:
 	template<class LoadIdBinder>
 	FLoadStructPlan CreatePlanInner(FLoadIdMapping Mapping, LoadIdBinder& LoadIds)
 	{
-		const FStructSchema& From = ResolveStructSchema(ReadId, Mapping.ReadId);
+		const FStructSchema& From = ResolveStructSchema(BatchId, Mapping.ReadId);
 		
 		TConstArrayView<FInnerStruct> LoweredInners;
 		if constexpr (std::is_same_v<LoadIdBinder, FLoadIdBinderDummy>)
@@ -832,9 +831,9 @@ private:
 	}
 };
 
-FLoadBatchPtr CreateLoadPlans(FReadBatchId ReadId, const FDeclarations& Declarations, const FCustomBindings& Customs, const FSchemaBindings& Schemas, TConstArrayView<FStructId> RuntimeIds, ESchemaFormat Format)
+FLoadBatchPtr CreateLoadPlans(FSchemaBatchId BatchId, const FDeclarations& Declarations, const FCustomBindings& Customs, const FSchemaBindings& Schemas, TConstArrayView<FStructId> RuntimeIds, ESchemaFormat Format)
 {
-	return FLoadPlanner{ReadId, Declarations, Customs, Schemas, RuntimeIds}.CreatePlans(Format);
+	return FLoadPlanner{BatchId, Declarations, Customs, Schemas, RuntimeIds}.CreatePlans(Format);
 }
 
 ////////////////////////////////////////////////////////////////////////////
@@ -1234,7 +1233,7 @@ FNestedRangeLoadView FRangeLoadView::AsRanges() const
 
 static FStructView ToReadView(FStructLoadView In)
 {
-	FStructSchemaHandle ReadSchema = { In.Schema.Batch.GetReadId(In.Schema.LoadId), In.Schema.Batch.ReadId };
+	FStructSchemaHandle ReadSchema = { In.Schema.Batch.GetReadId(In.Schema.LoadId), In.Schema.Batch.BatchId };
 	return { ReadSchema, In.Values };
 }
 
@@ -1298,7 +1297,7 @@ void ConstructAndLoadStruct(void* Dst, FStructLoadView Src)
 
 void FSchemaLoadHandle::GetInnerLoadIds(TArrayView<FOptionalSchemaId> Out) const
 {
-	const FStructSchema& ReadSchema = FStructSchemaHandle{ Batch.GetReadId(LoadId), Batch.ReadId }.Resolve();
+	const FStructSchema& ReadSchema = FStructSchemaHandle{ Batch.GetReadId(LoadId), Batch.BatchId }.Resolve();
 	TConstArrayView<FMemberType> MemberTypes = ReadSchema.GetMemberTypes();
 	TConstArrayView<FMemberType> RangeTypes = ReadSchema.GetRangeTypes();
 	check(Out.Num() == MemberTypes.Num());
diff --git a/Engine/Plugins/Experimental/PlainProps/Source/Private/PlainPropsParse.cpp b/Engine/Plugins/Experimental/PlainProps/Source/Private/PlainPropsParse.cpp
--- a/Engine/Plugins/Experimental/PlainProps/Source/Private/PlainPropsParse.cpp
+++ b/Engine/Plugins/Experimental/PlainProps/Source/Private/PlainPropsParse.cpp
@@ -622,7 +622,7 @@ public:
 		FDeclarations& InTypes,
 		FParsedSchemas& InSchemas);
 
-	FBuiltStruct* ParseMembers(FDeclId SchemaId);
+	FBuiltStruct* ParseMembers(FDeclId Id);
 
 private:
 	FDebugIds GetDebug() const { return FDebugIds(Names); }
@@ -636,7 +636,7 @@ private:
 	bool AdvanceToNextMember();
 	FDeclId ParseDynamicStructType();
 
-	[[nodiscard]] FBuiltStruct* BuildAndReset(FDeclId Schema);
+	[[nodiscard]] FBuiltStruct* BuildAndReset(FDeclId Id);
 
 	[[nodiscard]] FTypedRange	ParseLeaves(const FMemberSchema& RangeSchema, FScopedTokenReader RangeTokens);
 	[[nodiscard]] FTypedRange	ParseStructs(const FMemberSchema& RangeSchema, FScopedTokenReader RangeTokens);
diff --git a/Engine/Plugins/Experimental/PlainProps/Source/Private/PlainPropsPrint.cpp b/Engine/Plugins/Experimental/PlainProps/Source/Private/PlainPropsPrint.cpp
--- a/Engine/Plugins/Experimental/PlainProps/Source/Private/PlainPropsPrint.cpp
+++ b/Engine/Plugins/Experimental/PlainProps/Source/Private/PlainPropsPrint.cpp
@@ -254,7 +254,7 @@ struct FMemberSchemaView
 	using FMemberTypeRange = TConstArrayView<FMemberType>;
 
 	FMemberType				Type;
-	FReadBatchId			Batch;
+	FSchemaBatchId			Batch;
 	FOptionalSchemaId		InnerSchema;
 	FMemberTypeRange		InnerRangeTypes;
 
@@ -275,7 +275,7 @@ struct FMemberSchemaView
 class FStructSchemaReader
 {
 public:
-	FStructSchemaReader(const FStructSchema& Schema, FReadBatchId InBatch);
+	FStructSchemaReader(const FStructSchema& Schema, FSchemaBatchId InBatch);
 
 	FType					GetStruct() const		{ return Struct; }
 	bool					IsDense() const			{ return bIsDense; }
@@ -293,7 +293,7 @@ public:
 
 private:
 	const FMemberType*		Footer;
-	const FReadBatchId		Batch;					// Needed to resolve schemas
+	const FSchemaBatchId	Batch;					// Needed to resolve schemas
 	const FType				Struct;
 	const bool				bIsDense : 1;
 	const bool				bHasSuper : 1;
@@ -431,7 +431,7 @@ void PrintYamlBatch(FUtf8Builder& Out, const FBatchIds& Ids, TConstArrayView<FSt
 
 ///////////////////////////////////////////////////////////////////////////////
 
-FStructSchemaReader::FStructSchemaReader(const FStructSchema& Schema, FReadBatchId InBatch)
+FStructSchemaReader::FStructSchemaReader(const FStructSchema& Schema, FSchemaBatchId InBatch)
 : Footer(Schema.Footer)
 , Batch(InBatch)
 , Struct(Schema.Type)
@@ -807,7 +807,7 @@ public:
 	FUtf8StringView operator*() const	{ return Buffer.ToView(); }
 };
 
-void FBatchPrinter::PrintStructSchema(const FStructSchema& Struct, FReadBatchId BatchId)
+void FBatchPrinter::PrintStructSchema(const FStructSchema& Struct, FSchemaBatchId BatchId)
 {
 	FStructSchemaReader Reader(Struct, BatchId);
 
diff --git a/Engine/Plugins/Experimental/PlainProps/Source/Private/PlainPropsRead.cpp b/Engine/Plugins/Experimental/PlainProps/Source/Private/PlainPropsRead.cpp
--- a/Engine/Plugins/Experimental/PlainProps/Source/Private/PlainPropsRead.cpp
+++ b/Engine/Plugins/Experimental/PlainProps/Source/Private/PlainPropsRead.cpp
@@ -46,7 +46,7 @@ class FReadSchemaRegistry
 	std::atomic<uint32> Counter{};
 
 public:
-	FReadBatchId Mount(const FSchemaBatch* Batch)
+	FSchemaBatchId Mount(const FSchemaBatch* Batch)
 	{
 		const uint32 Count = Counter.fetch_add(1, std::memory_order_relaxed);
 		for (uint32 BaseIdx = 0; BaseIdx < FReadSchemaRegistry::Capacity; ++BaseIdx)
@@ -56,7 +56,7 @@ public:
 			if (Slots[Idx].load(std::memory_order_relaxed) == nullptr &&
 				Slots[Idx].compare_exchange_strong(Null, Batch, std::memory_order_acq_rel))
 			{
-				FReadBatchId Out;
+				FSchemaBatchId Out;
 				Out.Idx = static_cast<uint16>(Idx);
 				return Out;
 			}
@@ -66,14 +66,14 @@ public:
 		while (true);
 	}
 
-	const FSchemaBatch* Unmount(FReadBatchId Id)
+	const FSchemaBatch* Unmount(FSchemaBatchId Id)
 	{
 		const FSchemaBatch* Batch = &Get(Id);
 		verify(Slots[Id.Idx].compare_exchange_strong(Batch, nullptr, std::memory_order_acq_rel));
 		return Batch;
 	}
 
-	const FSchemaBatch& Get(FReadBatchId Id) const
+	const FSchemaBatch& Get(FSchemaBatchId Id) const
 	{
 		const FSchemaBatch* Batch = Slots[Id.Idx].load(std::memory_order_relaxed);
 		check(Batch);
@@ -89,42 +89,42 @@ const FSchemaBatch*	ValidateSchemas(FMemoryView Schemas)
 	return Batch;
 }
 
-FReadBatchId MountReadSchemas(const FSchemaBatch* Batch)
+FSchemaBatchId MountReadSchemas(const FSchemaBatch* Batch)
 {
 	return GReadSchemas.Mount(Batch);
 }
 
-const FSchemaBatch* UnmountReadSchemas(FReadBatchId Id)
+const FSchemaBatch* UnmountReadSchemas(FSchemaBatchId Id)
 {
 	return GReadSchemas.Unmount(Id);
 }
 
-uint32 NumStructSchemas(FReadBatchId Batch)
+uint32 NumStructSchemas(FSchemaBatchId Batch)
 {
 	return GReadSchemas.Get(Batch).NumStructSchemas;
 }
 
-const FSchemaBatch& GetReadSchemas(FReadBatchId Batch)
+const FSchemaBatch& GetReadSchemas(FSchemaBatchId Batch)
 {
 	return GReadSchemas.Get(Batch);
 }
 
-const FStructSchema& ResolveStructSchema(FReadBatchId Batch, FStructSchemaId Schema)
+const FStructSchema& ResolveStructSchema(FSchemaBatchId Batch, FStructSchemaId Schema)
 {
 	return ResolveStructSchema(GReadSchemas.Get(Batch), Schema);
 }
 
-const FEnumSchema& ResolveEnumSchema(FReadBatchId Batch, FEnumSchemaId Schema)
+const FEnumSchema& ResolveEnumSchema(FSchemaBatchId Batch, FEnumSchemaId Schema)
 {
 	return ResolveEnumSchema(GReadSchemas.Get(Batch), Schema);
 }
 
-FNestedScope ResolveUntranslatedNestedScope(FReadBatchId Batch, FNestedScopeId Id)
+FNestedScope ResolveUntranslatedNestedScope(FSchemaBatchId Batch, FNestedScopeId Id)
 {
 	return ResolveNestedScope(GReadSchemas.Get(Batch), Id);
 }
 
-FParametricTypeView ResolveUntranslatedParametricType(FReadBatchId Batch, FParametricTypeId Id)
+FParametricTypeView ResolveUntranslatedParametricType(FSchemaBatchId Batch, FParametricTypeId Id)
 {
 	return ResolveParametricType(GReadSchemas.Get(Batch), Id);
 }
@@ -180,7 +180,7 @@ const FMemberType*	FMemberReader::GetRangeTypes() const	{ return FStructSchema::
 const FSchemaId*	FMemberReader::GetInnerSchemas() const	{ return FStructSchema::GetInnerSchemas(Footer, NumMembers, NumRangeTypes, NumMembers - HasSuper); }
 const FMemberId*	FMemberReader::GetMemberNames() const	{ return FStructSchema::GetMemberNames(Footer, NumMembers, NumRangeTypes); }
 
-FMemberReader::FMemberReader(const FStructSchema& Schema, FByteReader Values, FReadBatchId InBatch)
+FMemberReader::FMemberReader(const FStructSchema& Schema, FByteReader Values, FSchemaBatchId InBatch)
 : Footer(Schema.Footer)
 , Batch(InBatch)
 , IsSparse(!Schema.IsDense)
@@ -419,7 +419,7 @@ FFlatMemberReader::FFlatMemberReader(FStructView Struct)
 
 //////////////////////////////////////////////////////////////////////////
 
-FBatchIds::FBatchIds(FReadBatchId Batch) : Schemas(GReadSchemas.Get(Batch)), BatchId(Batch) {}
+FBatchIds::FBatchIds(FSchemaBatchId Batch) : Schemas(GReadSchemas.Get(Batch)), BatchId(Batch) {}
 uint32 FBatchIds::NumEnums() const { return Schemas.NumSchemas - Schemas.NumStructSchemas; }
 uint32 FBatchIds::NumStructs() const { return Schemas.NumStructSchemas; }
 FType FBatchIds::Resolve(FEnumSchemaId Id) const { return ResolveEnumSchema(Schemas, Id).Type; }
diff --git a/Engine/Plugins/Experimental/PlainProps/Source/Public/PlainPropsBuild.h b/Engine/Plugins/Experimental/PlainProps/Source/Public/PlainPropsBuild.h
--- a/Engine/Plugins/Experimental/PlainProps/Source/Public/PlainPropsBuild.h
+++ b/Engine/Plugins/Experimental/PlainProps/Source/Public/PlainPropsBuild.h
@@ -150,21 +150,21 @@ FMemberSchema MakeLeafRangeSchema()
 }
 
 template<Enumeration T>
-FMemberSchema MakeEnumRangeSchema(FEnumId Schema, ERangeSizeType MaxSize)
+FMemberSchema MakeEnumRangeSchema(FEnumId Id, ERangeSizeType MaxSize)
 {
-	return { FMemberType(MaxSize), ReflectEnum<T>.Pack(), 1, FInnerId(Schema), nullptr };
+	return { FMemberType(MaxSize), ReflectEnum<T>.Pack(), 1, FInnerId(Id), nullptr };
 }
 
 template<Enumeration T, typename SizeType>
-FMemberSchema MakeEnumRangeSchema(FEnumId Schema)
+FMemberSchema MakeEnumRangeSchema(FEnumId Id)
 {
-	return MakeEnumRangeSchema<T>(Schema, RangeSizeOf(SizeType{}));
+	return MakeEnumRangeSchema<T>(Id, RangeSizeOf(SizeType{}));
 }
 
 template<UnsignedIntegral T>
-FMemberSchema MakeEnumRangeSchema(FEnumId Schema, ERangeSizeType MaxSize)
+FMemberSchema MakeEnumRangeSchema(FEnumId Id, ERangeSizeType MaxSize)
 {
-	return { FMemberType(MaxSize), FMemberType(ELeafType::Enum, ReflectArithmetic<T>.Width), 1, FInnerId(Schema), nullptr };
+	return { FMemberType(MaxSize), FMemberType(ELeafType::Enum, ReflectArithmetic<T>.Width), 1, FInnerId(Id), nullptr };
 }
 
 inline constexpr FMemberType DefaultStructType =	FMemberType(FStructType{EMemberKind::Struct, /* IsDynamic */ 0, /* IsSuper */ 0});
@@ -172,9 +172,9 @@ inline constexpr FMemberType SuperStructType =		FMemberType(FStructType{EMemberK
 inline constexpr FMemberType DynamicStructType =	FMemberType(FStructType{EMemberKind::Struct, /* IsDynamic */ 1, /* IsSuper */ 0});
 
 
-inline FMemberSchema MakeStructRangeSchema(ERangeSizeType SizeType, FStructId Schema)
+inline FMemberSchema MakeStructRangeSchema(ERangeSizeType SizeType, FStructId Id)
 {
-	return { FMemberType(SizeType), DefaultStructType, 1, FInnerId(Schema), nullptr };
+	return { FMemberType(SizeType), DefaultStructType, 1, FInnerId(Id), nullptr };
 }
 
 inline FMemberSchema MakeDynamicStructRangeSchema(ERangeSizeType SizeType)
@@ -225,9 +225,9 @@ template<UnsignedIntegral T>
 	return { MakeEnumRangeSchema<T>(Enum, SizeType), CloneLeaves(Scratch, Values.Num(), Values.GetData(), sizeof(T)) };
 }
 
-[[nodiscard]] inline FTypedRange MakeStructRange(FStructId Schema, ERangeSizeType SizeType, FBuiltRange* Values )
+[[nodiscard]] inline FTypedRange MakeStructRange(FStructId Id, ERangeSizeType SizeType, FBuiltRange* Values )
 {
-	return { MakeStructRangeSchema(SizeType, Schema), Values };
+	return { MakeStructRangeSchema(SizeType, Id), Values };
 }
 
 //////////////////////////////////////////////////////////////////////////
@@ -249,10 +249,10 @@ struct FBuiltMember
 {
 	FBuiltMember(FMemberId InName, FTypedValue In) : FBuiltMember(InName, In.Schema, In.Value) {}
 	FBuiltMember(FOptionalMemberId N, FMemberSchema S, FBuiltValue V) : Name(N), Schema(MoveTemp(S)), Value(V) {}
-	PLAINPROPS_API FBuiltMember(FMemberId Name, FUnpackedLeafType Leaf, FOptionalEnumId Schema, uint64 Value);
+	PLAINPROPS_API FBuiltMember(FMemberId Name, FUnpackedLeafType Leaf, FOptionalEnumId Id, uint64 Value);
 	PLAINPROPS_API FBuiltMember(FMemberId Name, FTypedRange Range);
-	PLAINPROPS_API FBuiltMember(FMemberId Name, FStructId Schema, FBuiltStruct* Value);
-	PLAINPROPS_API static FBuiltMember MakeSuper(FStructId Schema, FBuiltStruct* Value);
+	PLAINPROPS_API FBuiltMember(FMemberId Name, FStructId Id, FBuiltStruct* Value);
+	PLAINPROPS_API static FBuiltMember MakeSuper(FStructId Id, FBuiltStruct* Value);
 
 	FOptionalMemberId		Name;
 	FMemberSchema			Schema;
@@ -272,15 +272,15 @@ public:
 	}
 
 	template<Enumeration T>
-	void AddEnum(FMemberId Name, FEnumId Schema, T Value)
+	void AddEnum(FMemberId Name, FEnumId Id, T Value)
 	{
-		AddLeaf(Name, ReflectEnum<T>, ToOptional(Schema), ValueCast(static_cast<__underlying_type(T)>(Value)));
+		AddLeaf(Name, ReflectEnum<T>, ToOptional(Id), ValueCast(static_cast<__underlying_type(T)>(Value)));
 	}
 	
 	template<UnsignedIntegral T>
-	void AddEnum(FMemberId Name, FEnumId Schema, T Value)
+	void AddEnum(FMemberId Name, FEnumId Id, T Value)
 	{
-		AddLeaf(Name, {ELeafType::Enum, ReflectArithmetic<T>.Width}, ToOptional(Schema), Value);
+		AddLeaf(Name, {ELeafType::Enum, ReflectArithmetic<T>.Width}, ToOptional(Id), Value);
 	}
 	
 	template<UnsignedIntegral T>
@@ -292,7 +292,7 @@ public:
 	void AddLeaf(FMemberId Name, FUnpackedLeafType Leaf, FOptionalEnumId Enum, uint64 Value)	{ Members.Emplace(Name, Leaf, Enum, Value); }
 	void AddRange(FMemberId Name, FTypedRange Range)											{ Members.Emplace(Name, Range); }
 	// Add already built nested struct, must not be null
-	void AddStruct(FMemberId Name, FStructId Schema, FBuiltStruct* Struct)						{ check(Struct); Members.Emplace(Name, Schema, Struct); }
+	void AddStruct(FMemberId Name, FStructId Id, FBuiltStruct* Struct)							{ check(Struct); Members.Emplace(Name, Id, Struct); }
 	void Add(FMemberId Name, FTypedValue TypedValue)											{ Members.Emplace(Name, TypedValue); }
 
 	// Add already built super struct, must not be null and must be the first member added
diff --git a/Engine/Plugins/Experimental/PlainProps/Source/Public/PlainPropsDiff.h b/Engine/Plugins/Experimental/PlainProps/Source/Public/PlainPropsDiff.h
--- a/Engine/Plugins/Experimental/PlainProps/Source/Public/PlainPropsDiff.h
+++ b/Engine/Plugins/Experimental/PlainProps/Source/Public/PlainPropsDiff.h
@@ -49,6 +49,6 @@ struct FReadDiffNode
 struct FReadDiffPath : public TArray<FReadDiffNode, TInlineAllocator<16>> {};
 
 bool PLAINPROPS_API DiffStruct(FStructView A, FStructView B, FReadDiffPath& Out);
-bool PLAINPROPS_API DiffSchemas(FReadBatchId A, FReadBatchId B);
+bool PLAINPROPS_API DiffSchemas(FSchemaBatchId A, FSchemaBatchId B);
 
 } // namespace PlainProps
diff --git a/Engine/Plugins/Experimental/PlainProps/Source/Public/PlainPropsLoad.h b/Engine/Plugins/Experimental/PlainProps/Source/Public/PlainPropsLoad.h
--- a/Engine/Plugins/Experimental/PlainProps/Source/Public/PlainPropsLoad.h
+++ b/Engine/Plugins/Experimental/PlainProps/Source/Public/PlainPropsLoad.h
@@ -22,7 +22,7 @@ struct FLoadBatchDeleter
 using FLoadBatchPtr = TUniquePtr<FLoadBatch, FLoadBatchDeleter>;
 
 
-[[nodiscard]] PLAINPROPS_API FLoadBatchPtr	CreateLoadPlans(FReadBatchId ReadId, const FDeclarations& Declarations, const FCustomBindings& Customs, const FSchemaBindings& Schemas, TConstArrayView<FStructId> RuntimeIds, ESchemaFormat Format);
+[[nodiscard]] PLAINPROPS_API FLoadBatchPtr	CreateLoadPlans(FSchemaBatchId BatchId, const FDeclarations& Declarations, const FCustomBindings& Customs, const FSchemaBindings& Schemas, TConstArrayView<FStructId> RuntimeIds, ESchemaFormat Format);
 PLAINPROPS_API void							LoadStruct(void* Dst, FByteReader Src, FStructSchemaId LoadId, const FLoadBatch& Batch);
 PLAINPROPS_API void							ConstructAndLoadStruct(void* Dst, FByteReader Src, FStructSchemaId LoadId, const FLoadBatch& Batch);
 
diff --git a/Engine/Plugins/Experimental/PlainProps/Source/Public/PlainPropsPrint.h b/Engine/Plugins/Experimental/PlainProps/Source/Public/PlainPropsPrint.h
--- a/Engine/Plugins/Experimental/PlainProps/Source/Public/PlainPropsPrint.h
+++ b/Engine/Plugins/Experimental/PlainProps/Source/Public/PlainPropsPrint.h
@@ -11,7 +11,7 @@ namespace PlainProps
 class FBatchIds;
 struct FDiffPath;
 struct FMemberSchema;
-class FReadBatchId;
+class FSchemaBatchId;
 struct FReadDiffPath;
 struct FStructView;
 
diff --git a/Engine/Plugins/Experimental/PlainProps/Source/Public/PlainPropsRead.h b/Engine/Plugins/Experimental/PlainProps/Source/Public/PlainPropsRead.h
--- a/Engine/Plugins/Experimental/PlainProps/Source/Public/PlainPropsRead.h
+++ b/Engine/Plugins/Experimental/PlainProps/Source/Public/PlainPropsRead.h
@@ -22,7 +22,7 @@ struct FStructSchema;
 //////////////////////////////////////////////////////////////////////////
 
 // Represents a batch currently being read from
-class FReadBatchId
+class FSchemaBatchId
 {
 	friend class FReadSchemaRegistry;
 	uint16 Idx = 0;
@@ -30,14 +30,14 @@ class FReadBatchId
 
 // @param Schemas must outlive read batch
 PLAINPROPS_API const FSchemaBatch*		ValidateSchemas(FMemoryView Schemas);
-PLAINPROPS_API FReadBatchId				MountReadSchemas(const FSchemaBatch* Schemas);
-PLAINPROPS_API const FSchemaBatch*		UnmountReadSchemas(FReadBatchId Batch);
+PLAINPROPS_API FSchemaBatchId				MountReadSchemas(const FSchemaBatch* Schemas);
+PLAINPROPS_API const FSchemaBatch*		UnmountReadSchemas(FSchemaBatchId Batch);
 
-PLAINPROPS_API uint32					NumStructSchemas(FReadBatchId Batch);
-PLAINPROPS_API const FStructSchema&		ResolveStructSchema(FReadBatchId Batch, FStructSchemaId Id);
-PLAINPROPS_API const FEnumSchema&		ResolveEnumSchema(FReadBatchId Batch, FEnumSchemaId Id);
-PLAINPROPS_API FNestedScope				ResolveUntranslatedNestedScope(FReadBatchId Batch, FNestedScopeId Id);
-PLAINPROPS_API FParametricTypeView		ResolveUntranslatedParametricType(FReadBatchId Batch, FParametricTypeId Id);
+PLAINPROPS_API uint32					NumStructSchemas(FSchemaBatchId Batch);
+PLAINPROPS_API const FStructSchema&		ResolveStructSchema(FSchemaBatchId Batch, FStructSchemaId Id);
+PLAINPROPS_API const FEnumSchema&		ResolveEnumSchema(FSchemaBatchId Batch, FEnumSchemaId Id);
+PLAINPROPS_API FNestedScope				ResolveUntranslatedNestedScope(FSchemaBatchId Batch, FNestedScopeId Id);
+PLAINPROPS_API FParametricTypeView		ResolveUntranslatedParametricType(FSchemaBatchId Batch, FParametricTypeId Id);
 
 //////////////////////////////////////////////////////////////////////////
 
@@ -180,7 +180,7 @@ public:
 struct FStructSchemaHandle
 {
 	FStructSchemaId			Id;
-	FReadBatchId			Batch;
+	FSchemaBatchId			Batch;
 
 	const FStructSchema&	Resolve() const { return ResolveStructSchema(Batch, Id); }
 };
@@ -202,7 +202,7 @@ union FMemberValue
 struct FLeafView
 {
 	FUnpackedLeafType		Leaf;
-	FReadBatchId			Batch;
+	FSchemaBatchId			Batch;
 	FOptionalEnumSchemaId	Enum;
 	FMemberValue			Value;
 
@@ -288,7 +288,7 @@ using FNestedRangeView = TStructuralRangeView<FNestedRangeIterator>;
 struct FRangeSchema
 {
 	FMemberType				ItemType;
-	FReadBatchId			Batch; // Needed to resolve inner schema
+	FSchemaBatchId			Batch; // Needed to resolve inner schema
 	FOptionalSchemaId		InnermostSchema;
 	const FMemberType*		NestedItemTypes; // For nested ranges, can be out-of-bounds otherwise
 };
@@ -452,7 +452,7 @@ inline void TRangeView<T>::Copy(void* Dst, uint64 NumBytes) const
 class FLeafRangeView
 {
 public:
-	FLeafRangeView(FUnpackedLeafType Leaf, FReadBatchId InBatch, FOptionalEnumSchemaId InEnum, uint64 Num, const uint8* Data) 
+	FLeafRangeView(FUnpackedLeafType Leaf, FSchemaBatchId InBatch, FOptionalEnumSchemaId InEnum, uint64 Num, const uint8* Data) 
 	: Type(Leaf.Type)
 	, Width(Leaf.Width)
 	, Batch(InBatch)
@@ -508,7 +508,7 @@ public:
 private:
 	ELeafType						Type;
 	ELeafWidth						Width;
-	FReadBatchId					Batch;
+	FSchemaBatchId					Batch;
 	FOptionalEnumSchemaId			Enum;
 	uint64							NumItems;
 	const uint8*					Values;
@@ -600,7 +600,7 @@ class FMemberReader
 {
 public:
 	explicit FMemberReader(FStructView Struct) : FMemberReader(Struct.Schema.Resolve(), Struct.Values, Struct.Schema.Batch) {}
-	FMemberReader(const FStructSchema& Schema, FByteReader Values, FReadBatchId InBatch);
+	FMemberReader(const FStructSchema& Schema, FByteReader Values, FSchemaBatchId InBatch);
 
 	bool								HasMore() const			{ return MemberIdx < NumMembers; }
 	
@@ -626,7 +626,7 @@ public:
 
 protected: // for unit tests
 	const FMemberType*		Footer;
-	const FReadBatchId		Batch;					// Needed to resolve schemas
+	const FSchemaBatchId	Batch;					// Needed to resolve schemas
 	const bool				IsSparse : 1;
 	const bool				HasSuper : 1;
 	const uint32			NumMembers;
@@ -737,7 +737,7 @@ public:
 	using FIdsBase::Resolve;
 	using FIdsBase::AppendString;
 
-	PLAINPROPS_API FBatchIds(FReadBatchId Batch);
+	PLAINPROPS_API FBatchIds(FSchemaBatchId Batch);
 
 	PLAINPROPS_API virtual uint32	NumEnums() const;
 	PLAINPROPS_API virtual uint32	NumStructs() const;
@@ -748,10 +748,10 @@ public:
 	PLAINPROPS_API virtual void		AppendString(FUtf8Builder& Out, FStructSchemaId Struct) const;
 
 	const FSchemaBatch& GetSchemas() const { return Schemas; } 
-	FReadBatchId GetBatchId() const { return BatchId; } 
+	FSchemaBatchId GetBatchId() const { return BatchId; } 
 protected:
 	const FSchemaBatch& Schemas;
-	FReadBatchId BatchId;
+	FSchemaBatchId BatchId;
 };
 
 // Serialized id resolver for ESchemaFormat::InMemoryNames
@@ -761,7 +761,7 @@ public:
 	using FBatchIds::Resolve;
 	using FBatchIds::AppendString;
 
-	FMemoryBatchIds(FReadBatchId Batch, const FIdsBase& InNames) : FBatchIds(Batch), Names(InNames) {}
+	FMemoryBatchIds(FSchemaBatchId Batch, const FIdsBase& InNames) : FBatchIds(Batch), Names(InNames) {}
 
 	virtual uint32				NumNames() const override						{ return Names.NumNames(); }
 	virtual uint32				NumNestedScopes() const override				{ return Names.NumNestedScopes(); }
@@ -788,4 +788,4 @@ public:
 	PLAINPROPS_API virtual FParametricTypeView	Resolve(FParametricTypeId Id) const override final;
 };
 
-} // namespace PlainProps
+} // namespace PlainProps
\ No newline at end of file
diff --git a/Engine/Plugins/Experimental/PlainProps/Source/Tests/PlainPropsTests.cpp b/Engine/Plugins/Experimental/PlainProps/Source/Tests/PlainPropsTests.cpp
--- a/Engine/Plugins/Experimental/PlainProps/Source/Tests/PlainPropsTests.cpp
+++ b/Engine/Plugins/Experimental/PlainProps/Source/Tests/PlainPropsTests.cpp
@@ -310,14 +310,14 @@ class FTestBatchIds final : public FStableBatchIds
 {
 	const FTestNameReader& Names;
 public:
-	FTestBatchIds(const FTestNameReader& InNames, FReadBatchId Batch) : FStableBatchIds(Batch), Names(InNames) {}
+	FTestBatchIds(const FTestNameReader& InNames, FSchemaBatchId Batch) : FStableBatchIds(Batch), Names(InNames) {}
 	
 	using FStableBatchIds::AppendString;
 	virtual uint32 NumNames() const override { return Names.NumNames(); }
 	virtual void AppendString(FUtf8StringBuilderBase& Out, FNameId Name) const override { Out.Append(Names[Name]); }
 };
 
-[[nodiscard]] FReadBatchId ParseBatch(TArray64<uint8>& OutData, TArray<FStructView>& OutObjects, FUtf8StringView YamlView)
+[[nodiscard]] FSchemaBatchId ParseBatch(TArray64<uint8>& OutData, TArray<FStructView>& OutObjects, FUtf8StringView YamlView)
 {
 	// Parse yaml
 	ParseYamlBatch(OutData, YamlView);
@@ -327,7 +327,7 @@ public:
 	const uint32 SchemasSize = It.Grab<uint32>();
 	FMemoryView SchemasView = It.GrabSlice(SchemasSize);
 	const FSchemaBatch* Schemas = ValidateSchemas(SchemasView);
-	FReadBatchId Batch = MountReadSchemas(Schemas);
+	FSchemaBatchId Batch = MountReadSchemas(Schemas);
 
 	// Grab parsed objects
 	while (uint64 NumBytes = It.GrabVarIntU())
@@ -354,7 +354,7 @@ static void RoundtripText(const FBatchIds& BatchIds, TConstArrayView<FStructView
 	// Parse yaml
 	TArray64<uint8> Data;
 	TArray<FStructView> ParsedObjects;
-	FReadBatchId ParsedBatch = ParseBatch(Data, ParsedObjects, YamlView);
+	FSchemaBatchId ParsedBatch = ParseBatch(Data, ParsedObjects, YamlView);
 
 	// Diff schemas
 	CHECK_FALSE(DiffSchemas(BatchIds.GetBatchId(), ParsedBatch));
@@ -396,7 +396,7 @@ public:
 		uint32 SchemasSize = It.Grab<uint32>();
 		FMemoryView SchemasView = It.GrabSlice(SchemasSize);
 		const FSchemaBatch* Schemas = ValidateSchemas(SchemasView);
-		FReadBatchId Batch = MountReadSchemas(Schemas);
+		FSchemaBatchId Batch = MountReadSchemas(Schemas);
 		CHECK(It.Grab<uint32>() == TestMagics[2]);
 		
 		// Read objects
@@ -906,7 +906,7 @@ TEST_CASE_NAMED(FPlainPropsReadWriteTest, "System::Core::Serialization::PlainPro
 		{
 			CHECK(Objects.Num() == 1);
 
-			FReadBatchId Batch = Objects[0].Schema.Batch;
+			FSchemaBatchId Batch = Objects[0].Schema.Batch;
 			auto GetEnumName = [Batch, &Names](FLeafView Leaf) { return Names[ResolveEnumSchema(Batch, Leaf.Enum.Get()).Type.Name]; };
 			auto EqualEnumNames = [&Names](TConstArrayView<FNameId> Ids, TConstArrayView<const char*> Strings)
 			{
@@ -1565,7 +1565,7 @@ TEST_CASE_NAMED(FPlainPropsReadWriteTest, "System::Core::Serialization::PlainPro
 			CHECK(Names[FMemberReader(Objects[1]).GrabStruct().Schema.Resolve().Type.Name]  == "B4");
 			CHECK(Names[FMemberReader(Objects[2]).GrabStruct().Schema.Resolve().Type.Name]  == "B4");
 
-			FReadBatchId BatchId = Objects[0].Schema.Batch;
+			FSchemaBatchId BatchId = Objects[0].Schema.Batch;
 			const FStructSchema& B6 = Objects[0].Schema.Resolve();
 			const FStructSchema& B5 = ResolveStructSchema(BatchId, B6.GetSuperSchema().Get());
 			const FStructSchema& B4 = ResolveStructSchema(BatchId, B5.GetSuperSchema().Get());
@@ -1639,7 +1639,7 @@ TEST_CASE_NAMED(FPlainPropsReadWriteTest, "System::Core::Serialization::PlainPro
 		}, 
 		[](TConstArrayView<FStructView> Objects, const FTestNameReader& Names)
 		{	
-			FReadBatchId Batch = Objects[0].Schema.Batch;
+			FSchemaBatchId Batch = Objects[0].Schema.Batch;
 			FType S4T = Objects[0].Schema.Resolve().Type;
 
 			FNestedScope DoubleNested = ResolveUntranslatedNestedScope(Batch, S4T.Scope.AsNested());
diff --git a/Engine/Plugins/Experimental/PlainProps/Source/Tests/PlainPropsUeCoreTests.cpp b/Engine/Plugins/Experimental/PlainProps/Source/Tests/PlainPropsUeCoreTests.cpp
--- a/Engine/Plugins/Experimental/PlainProps/Source/Tests/PlainPropsUeCoreTests.cpp
+++ b/Engine/Plugins/Experimental/PlainProps/Source/Tests/PlainPropsUeCoreTests.cpp
@@ -271,7 +271,7 @@ class FStableNameBatchIds final : public FStableBatchIds
 {
 	TConstArrayView<FName> Names;
 public:
-	FStableNameBatchIds(FReadBatchId Batch, TConstArrayView<FName> InNames) : FStableBatchIds(Batch), Names(InNames) {}
+	FStableNameBatchIds(FSchemaBatchId Batch, TConstArrayView<FName> InNames) : FStableBatchIds(Batch), Names(InNames) {}
 	using FStableBatchIds::AppendString;
 
 	virtual uint32				NumNames() const override							{ return static_cast<uint32>(Names.Num()); }
@@ -285,7 +285,7 @@ public:
 	using FBatchIds::Resolve;
 	using FBatchIds::AppendString;
 
-	FTranslationBatchIds(FReadBatchId Batch, FIdBinding&& InBinding) : FBatchIds(Batch), Binding(InBinding) {}
+	FTranslationBatchIds(FSchemaBatchId Batch, FIdBinding&& InBinding) : FBatchIds(Batch), Binding(InBinding) {}
 
 	virtual uint32				NumNames() const override												{ return GNames.NumNames(); }
 	virtual void				AppendString(FUtf8Builder& Out, FNameId Name) const override			{ PlainProps::AppendString(Out, GNames.ResolveName(Name)); }
@@ -314,7 +314,7 @@ static void RoundtripText(const FBatchIds& BatchIds, TConstArrayView<FStructView
 	// Parse yaml
 	TArray64<uint8> Data;
 	TArray<FStructView> ParsedObjects;
-	FReadBatchId ParsedBatch = ParseBatch(Data, ParsedObjects, YamlView);
+	FSchemaBatchId ParsedBatch = ParseBatch(Data, ParsedObjects, YamlView);
 
 	if (Format == ESchemaFormat::StableNames)
 	{
@@ -379,7 +379,7 @@ public:
 		CHECK(It.Grab<uint32>() == Magics[2]);
 
 
-		FReadBatchId Batch;
+		FSchemaBatchId Batch;
 		TOptional<FIdTranslator> RuntimeIds;
 		if (InRuntimeIds.IsEmpty())
 		{
@@ -446,7 +446,7 @@ public:
 		{
 			// Mount and use the saved schemas rather than the translated load schemas
 			// in order to verify that text roundtripping creates identical serialized ids.
-			FReadBatchId Batch = MountReadSchemas(SavedSchemas);
+			FSchemaBatchId Batch = MountReadSchemas(SavedSchemas);
 			TArray<FStructView> StableObjects(Objects);
 			for (FStructView& Struct : StableObjects)
 			{
diff --git a/Engine/Plugins/Experimental/PlainPropsUObject/Source/Private/PlainPropsUObjectBindings.cpp b/Engine/Plugins/Experimental/PlainPropsUObject/Source/Private/PlainPropsUObjectBindings.cpp
--- a/Engine/Plugins/Experimental/PlainPropsUObject/Source/Private/PlainPropsUObjectBindings.cpp
+++ b/Engine/Plugins/Experimental/PlainPropsUObject/Source/Private/PlainPropsUObjectBindings.cpp
@@ -3709,7 +3709,7 @@ public:
 		const FSchemaBatch* SavedSchemas = ValidateSchemas(It.GrabSlice(SchemasSize));
 		check(It.Grab<uint32>() == Magics[2]);
 		
-		FReadBatchId Batch = MountReadSchemas(SavedSchemas);
+		FSchemaBatchId Batch = MountReadSchemas(SavedSchemas);
 
 		// Read objects
  		while (uint64 NumBytes = It.GrabVarIntU())
@@ -3967,14 +3967,14 @@ class FStableNameBatchIds final : public FStableBatchIds
 {
 	TConstArrayView<FName> Names;
 public:
-	FStableNameBatchIds(FReadBatchId Batch, TConstArrayView<FName> InNames) : FStableBatchIds(Batch), Names(InNames) {}
+	FStableNameBatchIds(FSchemaBatchId Batch, TConstArrayView<FName> InNames) : FStableBatchIds(Batch), Names(InNames) {}
 	using FStableBatchIds::AppendString;
 
 	virtual uint32				NumNames() const override							{ return static_cast<uint32>(Names.Num()); }
 	virtual void				AppendString(FUtf8Builder& Out, FNameId Name) const override { Names[Name.Idx].AppendString(Out); }
 };
 
-[[nodiscard]] static FReadBatchId ParseBatch(
+[[nodiscard]] static FSchemaBatchId ParseBatch(
 	TArray64<uint8>& OutData,
 	TArray<FStructView>& OutObjects,
 	FUtf8StringView YamlView)
@@ -3987,7 +3987,7 @@ public:
 	const uint32 SchemasSize = It.Grab<uint32>();
 	FMemoryView SchemasView = It.GrabSlice(SchemasSize);
 	const FSchemaBatch* Schemas = ValidateSchemas(SchemasView);
-	FReadBatchId Batch = MountReadSchemas(Schemas);
+	FSchemaBatchId Batch = MountReadSchemas(Schemas);
 
 	// Grab parsed objects
 	while (uint64 NumBytes = It.GrabVarIntU())
@@ -4028,7 +4028,7 @@ static void RoundtripText(
 	UE_LOGFMT(LogPlainPropsUObject, Display, "Parsing PlainProps text using {Format}...", ToString(Format));
 	TArray64<uint8> Data;
 	TArray<FStructView> ParsedObjects;
-	FReadBatchId ParsedBatch = ParseBatch(Data, ParsedObjects, YamlView);
+	FSchemaBatchId ParsedBatch = ParseBatch(Data, ParsedObjects, YamlView);
 
 	if (Format == ESchemaFormat::StableNames)
 	{
@@ -4087,7 +4087,7 @@ public:
 		FMemoryView SavedSchemasView = It.GrabSlice(SchemasSize);
 		const FSchemaBatch* SavedSchemas = ValidateSchemas(SavedSchemasView);
 		check(It.Grab<uint32>() == Magics[2]);
-		FReadBatchId Batch = MountReadSchemas(SavedSchemas);
+		FSchemaBatchId Batch = MountReadSchemas(SavedSchemas);
 		
 		// Read objects
 		while (uint64 NumBytes = It.GrabVarIntU())
