PlainProps: Fix broken debug visualization for console(s).

#rnx

diff --git a/Engine/Plugins/Experimental/PlainProps/Source/Public/PlainPropsVisualize.h b/Engine/Plugins/Experimental/PlainProps/Source/Public/PlainPropsVisualize.h
--- a/Engine/Plugins/Experimental/PlainProps/Source/Public/PlainPropsVisualize.h
+++ b/Engine/Plugins/Experimental/PlainProps/Source/Public/PlainPropsVisualize.h
@@ -83,7 +83,7 @@ class FIdScope
 {
 public:
 	// NameType must be a unique static string identifer for NameT,
-	// it is used by natvis in a memcmp(str1,str2,N) expression to select the correct typed name indexer, N <= 8
+	// it is used by natvis in a strncmp(str1,str2,N) expression to select the correct typed name indexer, N <= 8
 	template<typename NameT>
 	FIdScope(const TIdIndexer<NameT>& Indexer, const char NameType[8] = nullptr)
 	: Current{ static_cast<const FIdIndexerBase*>(&Indexer), &NameType[0] }
