PlainProps LoadRange fix

#rnx

diff --git a/Engine/Plugins/Experimental/PlainProps/Source/Private/PlainPropsLoad.cpp b/Engine/Plugins/Experimental/PlainProps/Source/Private/PlainPropsLoad.cpp
--- a/Engine/Plugins/Experimental/PlainProps/Source/Private/PlainPropsLoad.cpp
+++ b/Engine/Plugins/Experimental/PlainProps/Source/Private/PlainPropsLoad.cpp
@@ -895,29 +895,28 @@ public:
 		ERangeSizeType Unused = ERangeSizeType::Uni;
 		FLoadRangePlan Plan = { Unused, InnermostStruct, InnerTypes, Bindings.GetData() };
 		
-		FByteReader ByteIt(Src.Values);
-		FBitCacheReader BitIt;
-		LoadRangePlan(Member, Src.NumItems, ByteIt, BitIt, Src.Schema.Batch, Plan);
-		ByteIt.CheckEmpty();
+		LoadRangePlan(Member, Src.NumItems, Src.Values, Src.Schema.Batch, Plan);
 	}
 
-	static void LoadRangePlan(uint8* Member, uint64 Num, FByteReader& ByteIt, FBitCacheReader& BitIt, const FLoadBatch& Batch, const FLoadRangePlan& Range)
+	// T is FByteReader& for normal load plans but FMemoryView for FRangeLoadView case where GrabRangeValues() is already called
+	template<typename T>
+	static void LoadRangePlan(uint8* Member, uint64 Num, T& ByteItOrValues, const FLoadBatch& Batch, const FLoadRangePlan& Range)
 	{
 		FRangeBinding Binding = Range.Bindings[0];
 		FMemberBindType InnerType = Range.InnerTypes[0];
 		
 		if (Binding.IsLeafBinding())
 		{
-			LoadLeafRange(Member, Num, Binding.AsLeafBinding(), ByteIt, UnpackNonBitfield(InnerType.AsLeaf()));
+			LoadLeafRange(Member, Num, Binding.AsLeafBinding(), ByteItOrValues, UnpackNonBitfield(InnerType.AsLeaf()));
 		}
 		else if (Num)
 		{
 			const IItemRangeBinding& ItemBinding = Binding.AsItemBinding();
 			switch (InnerType.GetKind())
 			{
-				case EMemberKind::Leaf:		LoadRangeValues(Member, Num, ItemBinding, ByteIt, Batch, UnpackNonBitfield(InnerType.AsLeaf())); break;
-				case EMemberKind::Range:	LoadRangeValues(Member, Num, ItemBinding, ByteIt, Batch, Range.Tail()); break;
-				case EMemberKind::Struct:	LoadRangeValues(Member, Num, ItemBinding, ByteIt, Batch, Range.InnermostStruct.Get()); break;
+				case EMemberKind::Leaf:		LoadRangeValues(Member, Num, ItemBinding, ByteItOrValues, Batch, UnpackNonBitfield(InnerType.AsLeaf())); break;
+				case EMemberKind::Range:	LoadRangeValues(Member, Num, ItemBinding, ByteItOrValues, Batch, Range.Tail()); break;
+				case EMemberKind::Struct:	LoadRangeValues(Member, Num, ItemBinding, ByteItOrValues, Batch, Range.InnermostStruct.Get()); break;
 			}
 		}
 		else
@@ -930,13 +929,24 @@ public:
 	static void LoadLeafRange(uint8* Member, uint64 Num, const ILeafRangeBinding& Binding, FByteReader& ByteIt, FUnpackedLeafType Leaf)
 	{
 		FMemoryView Values = Num ? GrabRangeValues(ByteIt, Num, Leaf) : FMemoryView();
+		LoadLeafRange(Member, Num, Binding, Values, Leaf);
+	}
+
+	static void LoadLeafRange(uint8* Member, uint64 Num, const ILeafRangeBinding& Binding, FMemoryView Values, FUnpackedLeafType Leaf)
+	{
 		Binding.LoadLeaves(Member, FLeafRangeLoadView(Values.GetData(), Num, Leaf));
 	}
 
 	template<class SchemaType>
 	static void LoadRangeValues(uint8* Member, uint64 Num, const IItemRangeBinding& Binding, FByteReader& ByteIt, const FLoadBatch& Batch, SchemaType&& Schema)
 	{
-		FByteReader ValueIt(GrabRangeValues(ByteIt, Num, Schema));
+		LoadRangeValues(Member, Num, Binding, GrabRangeValues(ByteIt, Num, Schema), Batch, Forward<SchemaType>(Schema));
+	}
+
+	template<class SchemaType>
+	static void LoadRangeValues(uint8* Member, uint64 Num, const IItemRangeBinding& Binding, FMemoryView Values, const FLoadBatch& Batch, SchemaType&& Schema)
+	{
+		FByteReader ValueIt(Values);
 		FBitCacheReader BitIt; // Only used by ranges of ERangeSizeType::Uni ranges
 		FLoadRangeContext Ctx{.Request = {Member, Num}};
 		
@@ -1017,7 +1027,7 @@ public:
 		for (uint8* It = Items.Data, *End = It + Items.NumBytes(); It != End; It += ItemSize)
 		{
 			uint64 Num = GrabRangeNum(Plan.MaxSize, ByteIt, BitIt);
-			LoadRangePlan(It, Num, ByteIt, BitIt, Batch, Plan);
+			LoadRangePlan(It, Num, ByteIt, Batch, Plan);
 		}
 	}
 };
@@ -1142,7 +1152,7 @@ private:
 	void LoadMemberRange(uint8* Member, const FLoadRangePlan& Plan)
 	{
 		uint64 Num = GrabRangeNum(Plan.MaxSize, ByteIt, BitIt);
-		LoadRangePlan(Member, Num, ByteIt, BitIt, Batch, Plan);
+		LoadRangePlan(Member, Num, ByteIt, Batch, Plan);
 	}
 };
 
