PlainProps: Improve natvis for types used in the custom bindings layer.

#rnx

diff --git a/Engine/Plugins/Experimental/PlainProps/Source/Private/PlainPropsRead.cpp b/Engine/Plugins/Experimental/PlainProps/Source/Private/PlainPropsRead.cpp
--- a/Engine/Plugins/Experimental/PlainProps/Source/Private/PlainPropsRead.cpp
+++ b/Engine/Plugins/Experimental/PlainProps/Source/Private/PlainPropsRead.cpp
@@ -2,6 +2,7 @@
 
 #include "PlainPropsRead.h"
 #include "PlainPropsInternalRead.h"
+#include "PlainPropsVisualize.h"
 #include "Misc/Optional.h"
 
 #include <atomic>
@@ -46,6 +47,11 @@ class FReadSchemaRegistry
 	std::atomic<uint32> Counter{};
 
 public:
+	FReadSchemaRegistry()
+	{
+		DbgVis::AssignReadSchemasDebuggingState(reinterpret_cast<DbgVis::FSchemaBatch**>(Slots));
+	}
+
 	FSchemaBatchId Mount(const FSchemaBatch* Batch)
 	{
 		const uint32 Count = Counter.fetch_add(1, std::memory_order_relaxed);
diff --git a/Engine/Plugins/Experimental/PlainProps/Source/Private/PlainPropsVisualize.cpp b/Engine/Plugins/Experimental/PlainProps/Source/Private/PlainPropsVisualize.cpp
--- a/Engine/Plugins/Experimental/PlainProps/Source/Private/PlainPropsVisualize.cpp
+++ b/Engine/Plugins/Experimental/PlainProps/Source/Private/PlainPropsVisualize.cpp
@@ -71,24 +71,42 @@ PP_DBGVIS_VERIFY_MEMBER(FEnumSchema, Width);
 PP_DBGVIS_VERIFY_MEMBER(FEnumSchema, Num);
 PP_DBGVIS_VERIFY_MEMBER(FEnumSchema, Footer);
 
+// Verify layouts for duplicated debug types for Load.cpp
+
+// PP_DBGVIS_VERIFY_TYPE(FLoadStructPlan);
+// PP_DBGVIS_VERIFY_MEMBER(FLoadStructPlan, Handle);
+//
+// PP_DBGVIS_VERIFY_TYPE(FLoadBatch);
+// PP_DBGVIS_VERIFY_MEMBER(FLoadBatch, BatchId);
+// PP_DBGVIS_VERIFY_MEMBER(FLoadBatch, NumReadSchemas);
+// PP_DBGVIS_VERIFY_MEMBER(FLoadBatch, NumPlans);
+// PP_DBGVIS_VERIFY_MEMBER(FLoadBatch, Plans);
+
 ////////////////////////////////////////////////////////////////////////////////
 
 namespace DbgVis
 {
 
-// D4B455B7-7BAB-4F1D-A944-98EC086FB4AB
-static constexpr FGuid GVisualizerGuid = FGuid(0xD4B455B7, 0x7BAB4F1D, 0xA94498EC, 0x086FB4AB);
+// D4B455B7-7BAB-4F1D-A944-98EC086FB4AB => d4b455b77bab4f1da94498ec086fb4ab
+static constexpr FGuid GIdVisualizerGuid = FGuid(0xD4B455B7, 0x7BAB4F1D, 0xA94498EC, 0x086FB4AB);
+// 0A05D5A9-DE4E-492D-989E-7F936CC1C843 => 0a05d5a9de4e492d989e7f936cc1c843
+static constexpr FGuid GBatchVisualizerGuid = FGuid(0x0A05D5A9, 0xDE4E492D, 0x989E7F93, 0x6CC1C843);
 
 void FIdScope::AssignDebuggingState()
 {
 	FIdScope::Global = &Current;
-	(void)::UE::Core::FVisualizerDebuggingState::Assign(GVisualizerGuid, Global);
+	(void)::UE::Core::FVisualizerDebuggingState::Assign(GIdVisualizerGuid, Global);
 }
 
 FIdScope::~FIdScope()
 {
 	FIdScope::Global = Previous;
-	(void)::UE::Core::FVisualizerDebuggingState::Assign(GVisualizerGuid, Global);
+	(void)::UE::Core::FVisualizerDebuggingState::Assign(GIdVisualizerGuid, Global);
+}
+
+void AssignReadSchemasDebuggingState(FSchemaBatch** Slots)
+{
+	(void)::UE::Core::FVisualizerDebuggingState::Assign(GBatchVisualizerGuid, Slots);
 }
 
 } // namespace DbgVis
diff --git a/Engine/Plugins/Experimental/PlainProps/Source/Public/PlainPropsTypes.h b/Engine/Plugins/Experimental/PlainProps/Source/Public/PlainPropsTypes.h
--- a/Engine/Plugins/Experimental/PlainProps/Source/Public/PlainPropsTypes.h
+++ b/Engine/Plugins/Experimental/PlainProps/Source/Public/PlainPropsTypes.h
@@ -195,9 +195,11 @@ struct FSchemaId
 
 // Serialized struct id of a runtime FStructId
 struct FStructSchemaId : FSchemaId {};
+static void KeepDebugInfo(FStructSchemaId*) {};
 
 // Serialized enum id of a runtime FEnumId
 struct FEnumSchemaId : FSchemaId {};
+static void KeepDebugInfo(FEnumSchemaId*) {};
 
 //////////////////////////////////////////////////////////////////////////
 
diff --git a/Engine/Plugins/Experimental/PlainProps/Source/Public/PlainPropsVisualize.h b/Engine/Plugins/Experimental/PlainProps/Source/Public/PlainPropsVisualize.h
--- a/Engine/Plugins/Experimental/PlainProps/Source/Public/PlainPropsVisualize.h
+++ b/Engine/Plugins/Experimental/PlainProps/Source/Public/PlainPropsVisualize.h
@@ -4,6 +4,7 @@
 
 #include "PlainPropsBuild.h"
 #include "PlainPropsIndex.h"
+#include "PlainPropsRead.h"
 #include "Containers/Set.h"
 
 namespace PlainProps
@@ -64,6 +65,22 @@ struct FEnumSchema
 	FNameId Footer[0];
 };
 
+
+// Load.cpp
+
+struct FLoadStructPlan
+{
+	uint64					Handle;
+};
+
+struct FLoadBatch
+{
+	FSchemaBatchId			BatchId;
+	uint32					NumReadSchemas;
+	uint32					NumPlans;
+	FLoadStructPlan			Plans[0];
+};
+
 ////////////////////////////////////////////////////////////////////////////////
 
 struct FIdVisualizer
@@ -77,6 +94,8 @@ struct FIdVisualizer
 	static void KeepDebugInfo(DbgVis::FSchemaBatch*) {};
 	static void KeepDebugInfo(DbgVis::FStructSchema*) {};
 	static void KeepDebugInfo(DbgVis::FEnumSchema*) {};
+	static void KeepDebugInfo(DbgVis::FLoadStructPlan*) {};
+	static void KeepDebugInfo(DbgVis::FLoadBatch*) {};
 };
 
 class FIdScope
@@ -102,6 +121,8 @@ private:
 	FIdVisualizer*					Previous;
 };
 
+void AssignReadSchemasDebuggingState(FSchemaBatch** Slots);
+
 } // namespace DbgVis
 
 } // namespace PlainProps
