PlainProps: Implement TVariant binding.

#rnx

diff --git a/Engine/Plugins/Experimental/PlainProps/Source/Private/PlainPropsUeCoreBindings.cpp b/Engine/Plugins/Experimental/PlainProps/Source/Private/PlainPropsUeCoreBindings.cpp
--- a/Engine/Plugins/Experimental/PlainProps/Source/Private/PlainPropsUeCoreBindings.cpp
+++ b/Engine/Plugins/Experimental/PlainProps/Source/Private/PlainPropsUeCoreBindings.cpp
@@ -9,6 +9,22 @@
 namespace PlainProps::UE
 {
 
+FMemberId FNumeralGenerator::Grow(int32 Max)
+{
+	check(Max >= Cache.Num());
+	int32 OldNum = Cache.Num();
+	Cache.SetNumUninitialized(Max + 1);
+
+	FName Numeral("_");
+	for (int32 Idx = OldNum; Idx <= Max; ++Idx)
+	{
+		Numeral.SetNumber(NAME_EXTERNAL_TO_INTERNAL(Idx));
+		Cache[Idx] = Indexer.NameMember(Numeral);
+	}
+
+	return Cache[Max];
+}
+
 void FTransformBinding::Save(FMemberBuilder& Dst, const FTransform& Src, const FTransform* Default, const FSaveContext& Ctx) const
 {
 	static_assert(std::is_same_v<decltype(FTransform().GetTranslation().X), double>);
diff --git a/Engine/Plugins/Experimental/PlainProps/Source/Public/PlainPropsRead.h b/Engine/Plugins/Experimental/PlainProps/Source/Public/PlainPropsRead.h
--- a/Engine/Plugins/Experimental/PlainProps/Source/Public/PlainPropsRead.h
+++ b/Engine/Plugins/Experimental/PlainProps/Source/Public/PlainPropsRead.h
@@ -206,7 +206,7 @@ struct FLeafView
 	FOptionalEnumSchemaId	Enum;
 	FMemberValue			Value;
 
-	FORCEINLINE bool		AsBool() const		{ return Value.bValue; }
+	FORCEINLINE bool		AsBool() const		{ return As<bool>(); }
 	FORCEINLINE int8		AsS8() const		{ return As<int8>(); }
 	FORCEINLINE uint8		AsU8() const		{ return As<uint8>(); }
 	FORCEINLINE int16		AsS16() const		{ return As<int16>(); }
@@ -233,6 +233,14 @@ struct FLeafView
 		check(Leaf == ExpectedLeaf);
 		return *reinterpret_cast<const T*>(Value.Ptr);
 	}
+
+	template<>
+	FORCEINLINE bool		As() const
+	{
+		check(Leaf == ReflectLeaf<bool>);
+		return Value.bValue;
+	}
+
 };
 
 //////////////////////////////////////////////////////////////////////////
@@ -788,4 +796,4 @@ public:
 	PLAINPROPS_API virtual FParametricTypeView	Resolve(FParametricTypeId Id) const override final;
 };
 
-} // namespace PlainProps
\ No newline at end of file
+} // namespace PlainProps
diff --git a/Engine/Plugins/Experimental/PlainProps/Source/Public/PlainPropsSaveMember.h b/Engine/Plugins/Experimental/PlainProps/Source/Public/PlainPropsSaveMember.h
--- a/Engine/Plugins/Experimental/PlainProps/Source/Public/PlainPropsSaveMember.h
+++ b/Engine/Plugins/Experimental/PlainProps/Source/Public/PlainPropsSaveMember.h
@@ -203,9 +203,9 @@ struct TRangeSerializer
 	template<typename ContextType>
 	bool DiffMember(const T& A, const T& B, FMemberId Name, ContextType& Ctx) const
 	{
-		bool bDiff = bLeafRange ? Bindings[0].AsLeafBinding().DiffLeaves(A, B)
+		bool bDiff = bLeafRange ? Bindings[0].AsLeafBinding().DiffLeaves(&A, &B)
 			: DiffItems(&A, &B, Bindings[0].AsItemBinding(), Ctx);
-		return Track(bDiff, Ctx, Types[0], Name, MakeDiffMetadata(), &A, &B);
+		return Track(bDiff, Ctx, BindTypes[0].Value, Name, MakeDiffMetadata(), &A, &B);
 	}
 		
 	FMemberSchema MakeMemberSchema() const
@@ -447,4 +447,4 @@ public:
 	using FRangeSaverBase::Finalize;
 };
 
-} // namespace PlainProps
\ No newline at end of file
+} // namespace PlainProps
diff --git a/Engine/Plugins/Experimental/PlainProps/Source/Public/PlainPropsUeCoreBindings.h b/Engine/Plugins/Experimental/PlainProps/Source/Public/PlainPropsUeCoreBindings.h
--- a/Engine/Plugins/Experimental/PlainProps/Source/Public/PlainPropsUeCoreBindings.h
+++ b/Engine/Plugins/Experimental/PlainProps/Source/Public/PlainPropsUeCoreBindings.h
@@ -2,11 +2,13 @@
 
 #pragma once
 
+#include "Algo/Find.h"
 #include "Containers/Set.h"
 #include "Math/Quat.h"
 #include "Math/Vector.h"
 #include "Math/Transform.h"
 #include "Misc/Optional.h"
+#include "Misc/TVariant.h"
 #include "PlainPropsBind.h"
 #include "PlainPropsIndex.h"
 #include "PlainPropsLoadMember.h"
@@ -18,6 +20,7 @@ namespace Verse { class FNativeString; }
 using FVerseString = Verse::FNativeString;
 
 PP_REFLECT_STRUCT_TEMPLATE(, TTuple, void, Key, Value); // Todo handle TTuple and higher arities
+PP_REFLECT_STRUCT_ONLY(0,, FEmptyVariantState, void);
 
 namespace UE::Math
 {
@@ -35,6 +38,29 @@ struct TIsContiguousContainer<PlainProps::TRangeView<T>>
 namespace PlainProps::UE
 {
 
+// Temporary implementation for numeral structs
+class FNumeralGenerator
+{
+	TIdIndexer<FName>&			Indexer;
+	TArray<FMemberId>			Cache; // 0, 1, 2, ..
+
+	PLAINPROPS_API FMemberId	Grow(int32 Max);
+
+public:
+	FNumeralGenerator(TIdIndexer<FName>& InIndexer) : Indexer(InIndexer) { }
+
+	inline FMemberId			Make(uint16 Numeral)
+	{
+		return int32(Numeral) < Cache.Num() ? Cache[Numeral] : Grow(Numeral);
+	}
+	
+	TConstArrayView<FMemberId>	MakeRange(uint16 Numerals)
+	{
+		Make(Numerals - 1);
+		return MakeArrayView(Cache.GetData(), Numerals);
+	}
+};
+
 //class FReflection
 //{
 //	TIdIndexer<FName>	Names;
@@ -899,69 +925,118 @@ struct TMapDeltaBinding : TBaseDeltaBinding<TMap<K, V, SetAllocator, KeyFuncs>,
 	}
 };
 
+struct FEmptyVariantStateBinding : ICustomBinding
+{
+	using Type = FEmptyVariantState;
+	static constexpr std::initializer_list<FMemberId> MemberIds = {};
 
-//	template <typename T, typename VariantType>
-//	struct TVariantConstructFromMember
-//	{
-//		/** Default construct the type and load it from the FArchive */
-//		static void Construct(VariantType& Dst, FMemberReader Src)
-//		{
-//			if constexpr (std::is_arithmetic_v<T>)
-//			{
-//				Dst.template Emplace<T>(Src.GrabLeaf().As<T>());
-//			}
-//			else constexpr 
-//
-//		}
-//	};
-//
-//template <typename... Ts>
-//struct TVariantBinding : ICustomBinding
-//{
-//	using VariantType = TVariant<Ts...>;
-//
-//	const FStructDeclaration& Declaration;
-//
-//	static constexpr void(*Loaders[])(FMemberReader&, VariantType&) = { &TLoader<Ts, VariantType>::Load... };
-//
-//	virtual void LoadStruct(void* Dst, FStructView Src, ECustomLoadMethod Method, const FLoadBatch& Batch) const override
-//	{
-//		VariantType& Variant = *reinterpret_cast<VariantType*>(Dst);
-//
-//		if (Method == ECustomLoadMethod::Assign)
-//		{
-//			Variant.~VariantType();
-//		}
-//
-//		FMemberReader Members(Src);
-//		const FMemberId* DeclaredName = Algo::Find(Declaration.Names, Members.PeekName());
-//		check(DeclaredName);
-//		int64 Idx = DeclaredName - &Declaration.Names[0];
-//
-//		check(TypeIndex < UE_ARRAY_COUNT(Loaders));
-//		Loaders[TypeIndex](Ar, OutVariant);		
-//		
-//	}
-//
-//	template<typename T>
-//	void Load(TVariant<Ts...>& Dst, FCustomMemberLoader& Src, ECustomLoadMethod Method)
-//	{
-//
-//		
-//		{
-//			Dst.Emplace(MoveTemp(*reinterpret_cast<T*>(Temp)));
-//		}
-//		else
-//		{
-//			new (Dst) TVariant<Ts...>(MoveTemp(*reinterpret_cast<T*>(Temp)));
-//		}
-//	}
-//
-//	virtual FBuiltStruct*	SaveStruct(const void* Src, FDebugIds Debug) const override
-//	{
-//		...
-//	}
-//};
+	template<class Ids> FEmptyVariantStateBinding(TCustomInit<Ids>) {}
+	void Save(FMemberBuilder& Dst, const FEmptyVariantState& Src, const FEmptyVariantState*, const FSaveContext&) const { }
+	void Load(FEmptyVariantState& Dst, FStructLoadView Src, ECustomLoadMethod Method) const { }
+	static bool Diff(const FEmptyVariantState& A, const FEmptyVariantState& B, const FBindContext&) { return false; }
+};
+
+template <typename T, typename... Ts>
+struct TVariantMemberSerializer
+{
+	using Type = TVariant<Ts...>;
+	using TSerializers = std::tuple<TMemberSerializer<Ts>...>;
+
+	static FTypedValue SaveMember(const Type& Src, const TSerializers& Serializers, const FSaveContext& Ctx)
+	{
+		return std::get<TMemberSerializer<T>>(Serializers).SaveMember(Src.template Get<T>(), Ctx);
+	}
+
+	static void LoadMember(Type& Dst, const TSerializers& Serializers, FMemberLoader& Members)
+	{
+		new (&Dst) Type(::UE::Core::Private::TUninitializedType<T>());
+		std::get<TMemberSerializer<T>>(Serializers).ConstructAndLoadMember(&Dst.template Get<T>(), Members);
+	}
+
+	template<typename ContextType>
+	static bool DiffMember(const Type& A, const Type& B, const TSerializers& Serializers, FMemberId MemberId, ContextType& Ctx)
+	{
+		check(A.GetIndex() == B.GetIndex());
+		return std::get<TMemberSerializer<T>>(Serializers).DiffMember(A.template Get<T>(), B.template Get<T>(), MemberId, Ctx);
+	}
+};
+
+template <typename... Ts>
+struct TVariantMemberFunctions
+{
+	using Type = TVariant<Ts...>;
+	using TSerializers = std::tuple<TMemberSerializer<Ts>...>;
+
+	FTypedValue(*SaveMember)(const Type&, const TSerializers&, const FSaveContext&);
+	void(*LoadMember)(Type&, const TSerializers&, FMemberLoader&);
+	bool(*DiffMember)(const Type&, const Type&, const TSerializers&, FMemberId, const FBindContext&);
+	bool(*DiffAndTrackMember)(const Type&, const Type&, const TSerializers&, FMemberId, FDiffContext&);
+};
+
+template <typename... Ts>
+struct TVariantBinding : ICustomBinding
+{
+	using Type = TVariant<Ts...>;
+	using TSerializers = std::tuple<TMemberSerializer<Ts>...>;
+
+	const FMemberId MemberIds[sizeof...(Ts)];
+	const TSerializers Serializers;
+	static constexpr TVariantMemberFunctions<Ts...> Functions[sizeof...(Ts)] = {{
+		&TVariantMemberSerializer<Ts,Ts...>::SaveMember,
+		&TVariantMemberSerializer<Ts,Ts...>::LoadMember,
+		&TVariantMemberSerializer<Ts,Ts...>::DiffMember,
+		&TVariantMemberSerializer<Ts,Ts...>::DiffMember}...};
+
+	template<class Ids>
+	TVariantBinding(TCustomInit<Ids> Init)
+	: MemberIds{Ids::IndexNumeral(::UE::Core::Private::TParameterPackTypeIndex<Ts, Ts...>::Value)...}
+	, Serializers{{Init, MakeArrayView({MemberIds[::UE::Core::Private::TParameterPackTypeIndex<Ts, Ts...>::Value]})}...}
+	{ }
+
+	void Save(FMemberBuilder& Dst, const Type& Src, const Type* Default, const FSaveContext& Ctx) const
+	{
+		uint8 MemberIdx = IntCastChecked<uint8>(Src.GetIndex());
+		check(MemberIdx < UE_ARRAY_COUNT(Functions));
+		FTypedValue Value = Functions[MemberIdx].SaveMember(Src, Serializers, Ctx);
+		Dst.Add(MemberIds[MemberIdx], Value);
+	}
+
+	void Load(Type& Dst, FStructLoadView Src, ECustomLoadMethod Method) const
+	{
+		FMemberLoader Members(Src);
+		const FMemberId* DeclaredName = Algo::Find(MemberIds, Members.PeekName());
+		check(DeclaredName);
+		uint8 MemberIdx = IntCastChecked<uint8>(DeclaredName - &MemberIds[0]);
+		check(MemberIdx < UE_ARRAY_COUNT(Functions));
+		if (Method == ECustomLoadMethod::Assign)
+		{
+			Dst.~Type();
+		}
+		Functions[MemberIdx].LoadMember(Dst, Serializers, Members);
+	}
+
+	bool Diff(const Type& A, const Type& B, const FBindContext& Ctx) const
+	{
+		uint8 MemberIdx = IntCastChecked<uint8>(A.GetIndex());
+		check(MemberIdx < UE_ARRAY_COUNT(Functions));
+		if (MemberIdx != B.GetIndex())
+		{
+			return true;
+		}
+		return Functions[MemberIdx].DiffMember(A, B, Serializers, MemberIds[MemberIdx], Ctx);
+	}
+
+	bool Diff(const Type& A, const Type& B, FDiffContext& Ctx) const
+	{
+		uint8 MemberIdx = IntCastChecked<uint8>(A.GetIndex());
+		check(MemberIdx < UE_ARRAY_COUNT(Functions));
+		if (MemberIdx != B.GetIndex())
+		{
+			return true;
+		}
+		return Functions[MemberIdx].DiffAndTrackMember(A, B, Serializers, MemberIds[MemberIdx], Ctx);
+	}
+};
 
 } // namespace PlainProps::UE
 
@@ -984,6 +1059,13 @@ struct TTypename<TPair<K,V>>
 	using Parameters = std::tuple<K, V>;
 };
 
+template <typename... Ts>
+struct TTypename<TVariant<Ts...>>
+{
+	inline static constexpr std::string_view DeclName = "Variant";
+	using Parameters = std::tuple<Ts...>;
+};
+
 inline static constexpr std::string_view UeArrayName = "Array";
 inline static constexpr std::string_view UeSetName = "Set";
 inline static constexpr std::string_view UeMapName = "Map";
@@ -1096,4 +1178,12 @@ struct TCustomBind<TPair<K, V>>
 	using Type = UE::TPairBinding<K, V>;
 };
 
+template<> struct TCustomBind<FEmptyVariantState> { using Type = UE::FEmptyVariantStateBinding; };
+
+template<typename... Ts>
+struct TCustomBind<TVariant<Ts...>>
+{
+	using Type = UE::TVariantBinding<Ts...>;
+};
+
 } // namespace PlainProps
diff --git a/Engine/Plugins/Experimental/PlainProps/Source/Tests/PlainPropsUeCoreTests.cpp b/Engine/Plugins/Experimental/PlainProps/Source/Tests/PlainPropsUeCoreTests.cpp
--- a/Engine/Plugins/Experimental/PlainProps/Source/Tests/PlainPropsUeCoreTests.cpp
+++ b/Engine/Plugins/Experimental/PlainProps/Source/Tests/PlainPropsUeCoreTests.cpp
@@ -39,9 +39,11 @@ static FDeclarations			GTypes(GDebug);
 static FSchemaBindings			GSchemas(GDebug);
 static FCustomBindingsBottom	GCustoms(GDebug);
 static FCustomBindingsOverlay	GDeltaCustoms(GCustoms);
+static FNumeralGenerator		GNumeralGenerator(GNames);
 
 struct FRuntimeIds
 {
+	static FMemberId			IndexNumeral(uint16 Numeral)		{ return GNumeralGenerator.Make(Numeral); }
 	static FNameId				IndexName(FAnsiStringView Name)		{ return GNames.MakeName(FName(Name)); }
 	static FMemberId			IndexMember(FAnsiStringView Name)	{ return GNames.NameMember(FName(Name)); }
 	static FConcreteTypenameId	IndexTypename(FAnsiStringView Name)	{ return GNames.NameType(FName(Name)); }
@@ -747,6 +749,17 @@ PP_REFLECT_STRUCT(PlainProps::UE::Test, FIntAlias, void, X);
 PP_REFLECT_STRUCT(PlainProps::UE::Test, FSame1, void, X);
 PP_REFLECT_STRUCT(PlainProps::UE::Test, FSame2, void, X);
 
+struct FVariants
+{
+	TVariant<bool, float> BoolFloat;
+	TVariant<float, FInt> FloatInt;
+	TVariant<FEmptyVariantState, FNDC> NDC;
+	TVariant<float, TArray<FInt>> FloatIntArray;
+	TVariant<int32, TVariant<int32, FNDC>> IntVariantNDC;
+};
+
+PP_REFLECT_STRUCT(PlainProps::UE::Test, FVariants, void, BoolFloat, FloatInt, NDC, FloatIntArray, IntVariantNDC);
+
 } namespace PlainProps {
 
 template <> struct TTypename<PlainProps::UE::Test::FIntAlias>
@@ -1275,6 +1288,67 @@ TEST_CASE_NAMED(FPlainPropsUeCoreTest, "System::Core::Serialization::PlainProps:
 			});
 	}
 
+	SECTION("TVariant")
+	{
+		TScopedDefaultStructBinding<FInt> Int;
+		TScopedDefaultStructBinding<FNDC> NDC;
+		TScopedDefaultStructBinding<FEmptyVariantState> EmptyVariantState;
+		TScopedDefaultStructBinding<TVariant<int32, FNDC>> InnerVariant;
+		TScopedDefaultStructBinding<FVariants> Variants;
+		Run([](FBatchSaver& Batch)
+			{
+				FVariants Var1 {
+					TVariant<bool, float>{TInPlaceType<bool>(), true},
+					TVariant<float, FInt>{TInPlaceType<float>(), 2.2f},
+					TVariant<FEmptyVariantState, FNDC>{},
+					TVariant<float, TArray<FInt>>{TInPlaceType<float>(), 4.4f},
+					TVariant<int32, TVariant<int32, FNDC>>{TInPlaceType<int32>(), 5}
+				};
+				FVariants Var2 {
+					TVariant<bool, float>{TInPlaceType<float>(), 1.1f},
+					TVariant<float, FInt>{TInPlaceType<FInt>(), 2},
+					TVariant<FEmptyVariantState, FNDC>{TInPlaceType<FNDC>(), 3},
+					TVariant<float, TArray<FInt>>{TInPlaceType<TArray<FInt>>(), MakeArrayView({FInt{4}, FInt{5}})},
+					TVariant<int32, TVariant<int32, FNDC>>{TInPlaceType<TVariant<int32, FNDC>>(), TVariant<int32, FNDC>{TInPlaceType<FNDC>(), 6}}
+				};
+				Batch.Save(Var1);
+				Batch.Save(Var2);
+			}, 
+			[](FBatchLoader& Batch)
+			{
+				FVariants Var1 {
+					TVariant<bool, float>{TInPlaceType<bool>(), true},
+					TVariant<float, FInt>{TInPlaceType<float>(), 2.2f},
+					TVariant<FEmptyVariantState, FNDC>{},
+					TVariant<float, TArray<FInt>>{TInPlaceType<float>(), 4.4f},
+					TVariant<int32, TVariant<int32, FNDC>>{TInPlaceType<int32>(), 5}
+				};
+				FVariants Var2 {
+					TVariant<bool, float>{TInPlaceType<float>(), 1.1f},
+					TVariant<float, FInt>{TInPlaceType<FInt>(), 2},
+					TVariant<FEmptyVariantState, FNDC>{TInPlaceType<FNDC>(), 3},
+					TVariant<float, TArray<FInt>>{TInPlaceType<TArray<FInt>>(), MakeArrayView({FInt{4}, FInt{5}})},
+					TVariant<int32, TVariant<int32, FNDC>>{TInPlaceType<TVariant<int32, FNDC>>(), TVariant<int32, FNDC>{TInPlaceType<FNDC>(), 6}}
+				};
+
+				FVariants LoadedVar1 = Batch.Load<FVariants>();
+				FVariants LoadedVar2 = Batch.Load<FVariants>();
+
+				CHECK(LoadedVar1.BoolFloat.Get<bool>()				== Var1.BoolFloat.Get<bool>());
+				CHECK(LoadedVar1.FloatInt.Get<float>()				== Var1.FloatInt.Get<float>());
+				CHECK(LoadedVar1.NDC.IsType<FEmptyVariantState>());
+				CHECK(LoadedVar1.FloatIntArray.Get<float>()			== Var1.FloatIntArray.Get<float>());
+				CHECK(LoadedVar1.IntVariantNDC.Get<int32>()			== Var1.IntVariantNDC.Get<int32>());
+
+				CHECK(LoadedVar2.BoolFloat.Get<float>()				== Var2.BoolFloat.Get<float>());
+				CHECK(LoadedVar2.FloatInt.Get<FInt>()				== Var2.FloatInt.Get<FInt>());
+				CHECK(LoadedVar2.NDC.Get<FNDC>()					== Var2.NDC.Get<FNDC>());
+				CHECK(LoadedVar2.FloatIntArray.Get<TArray<FInt>>()	== Var2.FloatIntArray.Get<TArray<FInt>>());
+				CHECK(LoadedVar2.IntVariantNDC.Get<TVariant<int32,FNDC>>().Get<FNDC>() ==
+						    Var2.IntVariantNDC.Get<TVariant<int32,FNDC>>().Get<FNDC>());
+			});
+	}
+
 	SECTION("Reference")
 	{}
 }
diff --git a/Engine/Plugins/Experimental/PlainPropsUObject/Source/Private/PlainPropsUObjectBindings.cpp b/Engine/Plugins/Experimental/PlainPropsUObject/Source/Private/PlainPropsUObjectBindings.cpp
--- a/Engine/Plugins/Experimental/PlainPropsUObject/Source/Private/PlainPropsUObjectBindings.cpp
+++ b/Engine/Plugins/Experimental/PlainPropsUObject/Source/Private/PlainPropsUObjectBindings.cpp
@@ -256,24 +256,6 @@ const void* FDefaultStructs::Get(FBindId Id)
 
 ////////////////////////////////////////////////////////////////////////////////////////////////
 
-FMemberId FNumeralGenerator::Grow(int32 Max)
-{
-	check(Max > Cache.Num());
-	int32 OldNum = Cache.Num();
-	Cache.SetNumUninitialized(Max + 1);
-
-	FName Numeral("_");
-	for (int32 Idx = OldNum; Idx <= Max; ++Idx)
-	{
-		Numeral.SetNumber(NAME_EXTERNAL_TO_INTERNAL(Idx));
-		Cache[Idx] = GUE.Names.NameMember(Numeral);
-	}
-
-	return Cache[Max];
-}
-
-////////////////////////////////////////////////////////////////////////////////////////////////
-
 FCommonScopeIds::FCommonScopeIds(TIdIndexer<FName>& Names)
 : Core(						Names.MakeScope("/Script/Core"))
 , CoreUObject(				Names.MakeScope("/Script/CoreUObject"))
@@ -342,6 +324,7 @@ FGlobals::FGlobals()
 , Structs(Scopes, Names)
 , Typenames(Names)
 , Members(Names)
+, Numerals(Names)
 , Debug(Names)
 {}
 
diff --git a/Engine/Plugins/Experimental/PlainPropsUObject/Source/Public/PlainPropsUObjectRuntime.h b/Engine/Plugins/Experimental/PlainPropsUObject/Source/Public/PlainPropsUObjectRuntime.h
--- a/Engine/Plugins/Experimental/PlainPropsUObject/Source/Public/PlainPropsUObjectRuntime.h
+++ b/Engine/Plugins/Experimental/PlainPropsUObject/Source/Public/PlainPropsUObjectRuntime.h
@@ -42,26 +42,6 @@ private:
 	void									ReserveFlags(uint32 Idx);
 };
 
-// Temporary implementation for numeral structs
-class FNumeralGenerator
-{
-	TArray<FMemberId>			Cache; // 0, 1, 2, ..
-
-	FMemberId					Grow(int32 Max);
-
-public:
-	inline FMemberId			Make(uint16 Numeral)
-	{
-		return int32(Numeral) < Cache.Num() ? Cache[Numeral] : Grow(Numeral);
-	}
-	
-	TConstArrayView<FMemberId>	MakeRange(uint16 Numerals)
-	{
-		Make(Numerals);
-		return MakeArrayView(Cache.GetData(), Numerals);
-	}
-};
-
 struct FCommonScopeIds
 {
 	explicit FCommonScopeIds(TIdIndexer<FName>& Names);
@@ -274,4 +254,4 @@ template<> inline FScopeId IndexNamespaceId<UE::FRuntimeIds, TTypename<FGuid>>()
 template<> inline FScopeId IndexNamespaceId<UE::FRuntimeIds, TTypename<FColor>>()			{ return UE::GUE.Scopes.Core; }
 template<> inline FScopeId IndexNamespaceId<UE::FRuntimeIds, TTypename<FLinearColor>>()		{ return UE::GUE.Scopes.Core; }
 
-}
\ No newline at end of file
+}
