PlainProps add member literal indexing for upgrades

#rnx

diff --git a/Engine/Plugins/Experimental/PlainProps/Source/Private/PlainPropsPrint.cpp b/Engine/Plugins/Experimental/PlainProps/Source/Private/PlainPropsPrint.cpp
--- a/Engine/Plugins/Experimental/PlainProps/Source/Private/PlainPropsPrint.cpp
+++ b/Engine/Plugins/Experimental/PlainProps/Source/Private/PlainPropsPrint.cpp
@@ -56,13 +56,13 @@ void FIdIndexerBase::InitParameterNames()
 	{
 		for (uint32 W = 0; W < 4; ++W)
 		{
-			Leaves[T][W] = InitParameterName(GLiterals.Leaves[T][W].GetData());
+			Leaves[T][W] = InitParameterName(GLiterals.Leaves[T][W]);
 		}
 	}
 
 	for (uint32 S = 0; S < 9; ++S)
 	{
-		Ranges[S] = InitParameterName(GLiterals.Ranges[S].GetData());
+		Ranges[S] = InitParameterName(GLiterals.Ranges[S]);
 	}
 }
 
diff --git a/Engine/Plugins/Experimental/PlainProps/Source/Public/PlainPropsIndex.h b/Engine/Plugins/Experimental/PlainProps/Source/Public/PlainPropsIndex.h
--- a/Engine/Plugins/Experimental/PlainProps/Source/Public/PlainPropsIndex.h
+++ b/Engine/Plugins/Experimental/PlainProps/Source/Public/PlainPropsIndex.h
@@ -4,6 +4,7 @@
 
 #include "PlainPropsTypes.h"
 #include "Containers/Set.h"
+#include "Containers/StringView.h"
 
 namespace PlainProps
 {
@@ -75,7 +76,7 @@ public:
 	const FNestedScopeIndexer&			GetNestedScopes() const								{ return NestedScopes; }
 	const FParametricTypeIndexer&		GetParametricTypes() const							{ return ParametricTypes; }
 
-	FConcreteTypenameId					IndexRangeBindName(const char* Name)				{ return InitParameterName(Name); }
+	FConcreteTypenameId					IndexRangeBindName(FAnsiStringView Name)			{ return InitParameterName(Name); }
 protected:
 	FNestedScopeIndexer					NestedScopes;
 	FParametricTypeIndexer				ParametricTypes;
@@ -85,7 +86,11 @@ protected:
 	FConcreteTypenameId					Ranges[9];
 
 	PLAINPROPS_API void					InitParameterNames();
-	virtual FConcreteTypenameId			InitParameterName(const char* Name) = 0;
+	FConcreteTypenameId					InitParameterName(FAnsiStringView Name)				{ return {IndexLiteral(Name)}; }
+	FMemberId							IndexMemberName(FAnsiStringView Name)				{ return {IndexLiteral(Name)}; }
+
+private:
+	virtual FNameId						IndexLiteral(FAnsiStringView Name) = 0;
 };
 
 template<class NameType>
@@ -99,7 +104,8 @@ public:
 
 	TIdIndexer() { InitParameterNames(); }
 
-	template<typename T> FNameId				MakeName(T&& Name)						{ return { static_cast<uint32>(Names.Add(NameT(Name)).AsInteger()) }; }
+	template<int N>	FNameId						MakeName(const char (&Name)[N])			{ return MakeName(FAnsiStringView(Name, N - 1)); }
+	template<typename T> FNameId				MakeName(T&& Name)						{ return { static_cast<uint32>(Names.Add(NameT(Forward<T>(Name))).AsInteger()) }; }
 	template<typename T> FMemberId				NameMember(T&& Name)					{ return { MakeName(Forward<T>(Name)) }; }
 	template<typename T> FConcreteTypenameId	NameType(T&& Name)						{ return { MakeName(Forward<T>(Name)) }; }
 	template<typename T> FFlatScopeId			NameScope(T&& Name)						{ return { MakeName(Forward<T>(Name)) }; }
@@ -114,15 +120,8 @@ public:
 protected:
 	TSet<NameT>									Names;
 
-	virtual void								AppendString(FUtf8Builder& Out, FNameId Id) const override final
-	{
-		PlainProps::AppendString(Out, ResolveName(Id));
-	}
-
-	virtual FConcreteTypenameId					InitParameterName(const char* Name) override final
-	{ 
-		return NameType(Name);
-	}
+	virtual void	AppendString(FUtf8Builder& Out, FNameId Id) const override final	{ PlainProps::AppendString(Out, ResolveName(Id));}
+	virtual FNameId	IndexLiteral(FAnsiStringView Name) override final					{ return MakeName(NameT(Name)); }
 };
 
 
diff --git a/Engine/Plugins/Experimental/PlainProps/Source/Public/PlainPropsUeCoreBindings.h b/Engine/Plugins/Experimental/PlainProps/Source/Public/PlainPropsUeCoreBindings.h
--- a/Engine/Plugins/Experimental/PlainProps/Source/Public/PlainPropsUeCoreBindings.h
+++ b/Engine/Plugins/Experimental/PlainProps/Source/Public/PlainPropsUeCoreBindings.h
@@ -47,6 +47,7 @@ public:
 	using IntType = uint64;
 #endif
 	inline explicit FSensitiveName(FName Name);
+	inline explicit FSensitiveName(FAnsiStringView Name) : FSensitiveName(FName(Name)) {}
 	inline FName ToName() const { return FName(GetComparisonId(), GetDisplayId(), GetNumber()); }
 	inline void AppendString(FUtf8StringBuilderBase& Out) const { ToName().AppendString(Out); }
 	inline IntType ToUnstableInt() const;
diff --git a/Engine/Plugins/Experimental/PlainPropsUObject/Source/Private/PlainPropsUObjectBindings.cpp b/Engine/Plugins/Experimental/PlainPropsUObject/Source/Private/PlainPropsUObjectBindings.cpp
--- a/Engine/Plugins/Experimental/PlainPropsUObject/Source/Private/PlainPropsUObjectBindings.cpp
+++ b/Engine/Plugins/Experimental/PlainPropsUObject/Source/Private/PlainPropsUObjectBindings.cpp
@@ -2726,7 +2726,7 @@ static constexpr std::string_view TrivialOptionalName = "TrivialOptional";
 template<uint32 ItemSize>
 struct TTrivialOptionalBinding : IItemRangeBinding
 {
-	TTrivialOptionalBinding() : IItemRangeBinding(GUE.Names.IndexRangeBindName(Concat<TrivialOptionalName, HexString<ItemSize>>.data())) {}
+	TTrivialOptionalBinding() : IItemRangeBinding(GUE.Names.IndexRangeBindName(ToAnsiView(Concat<TrivialOptionalName, HexString<ItemSize>>))) {}
 
 	virtual void ReadItems(FSaveRangeContext& Ctx) const override { ReadBoolOptionalItem(Ctx, ItemSize); }
 	virtual void MakeItems(FLoadRangeContext& Ctx) const override { MakeBoolOptionalItem(Ctx, ItemSize); }
