PlainProps added literal indexer, extended FScratchAllocator and more in preparation of upgrades

#rnx

diff --git a/Engine/Plugins/Experimental/PlainProps/Source/Private/PlainPropsBind.cpp b/Engine/Plugins/Experimental/PlainProps/Source/Private/PlainPropsBind.cpp
--- a/Engine/Plugins/Experimental/PlainProps/Source/Private/PlainPropsBind.cpp
+++ b/Engine/Plugins/Experimental/PlainProps/Source/Private/PlainPropsBind.cpp
@@ -495,7 +495,7 @@ static void DeclareOrValidate(FStructDeclarationPtr& Out, FStructSpec Info, TCon
 	}
 
 	// Make header and allocate
-	FStructDeclaration Header = { 1, Info.Id, Info.Super, Info.Version, Num, Binding.NumInnerSchemas, Binding.NumInnerRanges, Info.Occupancy };
+	FStructDeclaration Header = { 1, Info.Id, Info.Super, Info.Version, Num, Binding.NumInnerRanges, Binding.NumInnerSchemas, Info.Occupancy };
 	FStructDeclaration* Ptr = new (FMemory::MallocZeroed(Header.CalculateSize())) FStructDeclaration;
 	FMemory::Memcpy(Ptr, &Header, sizeof(Header));
 
diff --git a/Engine/Plugins/Experimental/PlainProps/Source/Private/PlainPropsBuild.cpp b/Engine/Plugins/Experimental/PlainProps/Source/Private/PlainPropsBuild.cpp
--- a/Engine/Plugins/Experimental/PlainProps/Source/Private/PlainPropsBuild.cpp
+++ b/Engine/Plugins/Experimental/PlainProps/Source/Private/PlainPropsBuild.cpp
@@ -31,7 +31,7 @@ uint8* FScratchAllocator::AllocateInNewPage(SIZE_T Size, uint32 Alignment)
 	return Out;
 }
 
-FScratchAllocator::~FScratchAllocator()
+void FScratchAllocator::Reset()
 {
 	for (FPage* It = LastPage; It; )
 	{
@@ -39,6 +39,10 @@ FScratchAllocator::~FScratchAllocator()
 		FMemory::Free(It);
 		It = Prev;
 	}
+
+	Cursor = nullptr;
+	PageEnd = nullptr;
+	LastPage = nullptr;
 }
 
 FMemberType& FMemberSchema::EditInnermostType(FScratchAllocator& Scratch)
diff --git a/Engine/Plugins/Experimental/PlainProps/Source/Private/PlainPropsDeclare.cpp b/Engine/Plugins/Experimental/PlainProps/Source/Private/PlainPropsDeclare.cpp
--- a/Engine/Plugins/Experimental/PlainProps/Source/Private/PlainPropsDeclare.cpp
+++ b/Engine/Plugins/Experimental/PlainProps/Source/Private/PlainPropsDeclare.cpp
@@ -173,7 +173,7 @@ FStructDeclarationPtr Declare(FStructSpec In)
 	// Make header and allocate
 	uint16 NumMembers = IntCastChecked<uint16>(In.MemberNames.Num());
 	uint16 NumInnerIds = IntCastChecked<uint16>(InnerIds.Num());
-	FStructDeclaration Header{1, In.Id, In.Super, In.Version, NumMembers, NumInnerIds, IntCastChecked<uint16>(NumInnerRanges), In.Occupancy};
+	FStructDeclaration Header{1, In.Id, In.Super, In.Version, NumMembers, IntCastChecked<uint16>(NumInnerRanges), NumInnerIds, In.Occupancy};
 	FStructDeclaration* Out = new (FMemory::Malloc(Header.CalculateSize())) FStructDeclaration;
 	FMemory::Memcpy(Out, &Header, sizeof(Header));
 	
diff --git a/Engine/Plugins/Experimental/PlainProps/Source/Public/PlainPropsBuild.h b/Engine/Plugins/Experimental/PlainProps/Source/Public/PlainPropsBuild.h
--- a/Engine/Plugins/Experimental/PlainProps/Source/Public/PlainPropsBuild.h
+++ b/Engine/Plugins/Experimental/PlainProps/Source/Public/PlainPropsBuild.h
@@ -41,7 +41,7 @@ class FScratchAllocator
 public:
 	UE_NONCOPYABLE(FScratchAllocator);
 	FScratchAllocator() = default;
-	PLAINPROPS_API ~FScratchAllocator();
+	~FScratchAllocator() { Reset(); }
 
 	inline void* Allocate(SIZE_T Size, uint32 Alignment)
 	{
@@ -72,6 +72,8 @@ public:
 		}
 		return Out;
 	}
+
+	PLAINPROPS_API void Reset();
 };
 
 //////////////////////////////////////////////////////////////////////////
@@ -385,3 +387,8 @@ private:
 };
 
 } // namespace PlainProps
+
+inline void* operator new(std::size_t Size, std::align_val_t Align, PlainProps::FScratchAllocator& Scratch)
+{
+	return Scratch.Allocate(Size, static_cast<uint32>(Align));
+}
\ No newline at end of file
diff --git a/Engine/Plugins/Experimental/PlainProps/Source/Public/PlainPropsIndex.h b/Engine/Plugins/Experimental/PlainProps/Source/Public/PlainPropsIndex.h
--- a/Engine/Plugins/Experimental/PlainProps/Source/Public/PlainPropsIndex.h
+++ b/Engine/Plugins/Experimental/PlainProps/Source/Public/PlainPropsIndex.h
@@ -47,6 +47,7 @@ private:
 
 //////////////////////////////////////////////////////////////////////////
 
+// Allows indexes higher level ids from existing parts
 class FIdIndexerBase : public FIds
 {
 public:
@@ -75,8 +76,6 @@ public:
 	
 	const FNestedScopeIndexer&			GetNestedScopes() const								{ return NestedScopes; }
 	const FParametricTypeIndexer&		GetParametricTypes() const							{ return ParametricTypes; }
-
-	FConcreteTypenameId					IndexRangeBindName(FAnsiStringView Name)			{ return InitParameterName(Name); }
 protected:
 	FNestedScopeIndexer					NestedScopes;
 	FParametricTypeIndexer				ParametricTypes;
@@ -87,17 +86,29 @@ protected:
 
 	PLAINPROPS_API void					InitParameterNames();
 	FConcreteTypenameId					InitParameterName(FAnsiStringView Name)				{ return {IndexLiteral(Name)}; }
-	FMemberId							IndexMemberName(FAnsiStringView Name)				{ return {IndexLiteral(Name)}; }
 
 private:
+	friend class FLiteralIndexerBase;
 	virtual FNameId						IndexLiteral(FAnsiStringView Name) = 0;
 };
 
+// Also allows indexing new names from string literals
+class FLiteralIndexerBase : public FIdIndexerBase
+{
+public:
+	// Override to split scope into parts, default implementation assumes a single flat scope
+	virtual FScopeId					IndexScopeLiteral(FAnsiStringView Scope)			{ return FScopeId(FFlatScopeId{IndexLiteral(Scope)}); }
+	FConcreteTypenameId					IndexTypenameLiteral(FAnsiStringView Type)			{ return {IndexLiteral(Type)}; }
+	FMemberId							IndexMemberLiteral(FAnsiStringView Member)			{ return {IndexLiteral(Member)}; }
+	FNameId								IndexEnumFlagLiteral(FAnsiStringView Flag)			{ return IndexLiteral(Flag); }
+};
+
 template<class NameType>
 void AppendString(FUtf8StringBuilderBase& Out, const NameType& Str);
 
+// Indexes FNameId from unique NameT instances or NameT-constructible types
 template<class NameT>
-class TIdIndexer : public FIdIndexerBase
+class TIdIndexer : public FLiteralIndexerBase
 {
 public:
 	using FIdIndexerBase::AppendString;
@@ -117,6 +128,8 @@ public:
 	NameT										ResolveName(FNameId Id)	const			{ return Names.Get(FSetElementId::FromInteger(Id.Idx)); }
 	virtual uint32								NumNames() const override final			{ return static_cast<uint32>(Names.Num()); }
 
+	FConcreteTypenameId							IndexRangeBindName(FAnsiStringView Str) { return InitParameterName(Str); }
+
 protected:
 	TSet<NameT>									Names;
 
