Resolve warning C5321: nonstandard extension used: encoding '\x##' as a multi-byte utf-8 character. Use \u instead for cross platform compatibility or '/Zc:u8EscapeEncoding' to disable the extension.

#rnx

diff --git a/Engine/Plugins/Experimental/PlainProps/Source/Tests/PlainPropsTests.cpp b/Engine/Plugins/Experimental/PlainProps/Source/Tests/PlainPropsTests.cpp
--- a/Engine/Plugins/Experimental/PlainProps/Source/Tests/PlainPropsTests.cpp
+++ b/Engine/Plugins/Experimental/PlainProps/Source/Tests/PlainPropsTests.cpp
@@ -1097,7 +1097,7 @@ TEST_CASE_NAMED(FPlainPropsReadWriteTest, "System::Core::Serialization::PlainPro
 			Utf8.AddRange(Batch.NameMember("Null"),		BuildLeafRange(Scratch, TConstArrayView<char8_t>()));
 			Utf8.AddRange(Batch.NameMember("Empty"),	BuildLeafRange(Scratch, MakeArrayView(u8"")));
 			Utf8.AddRange(Batch.NameMember("Escape"),	BuildLeafRange(Scratch, MakeArrayView(u8"\"\\ \x01 \x1f\" \"\b \f \n \r \t \"\\").LeftChop(1)));
-			Utf8.AddRange(Batch.NameMember("Latin1"),	BuildLeafRange(Scratch, MakeArrayView(u8"\xE5 \xE4 \xF6")));		// å ä ö
+			Utf8.AddRange(Batch.NameMember("Latin1"),	BuildLeafRange(Scratch, MakeArrayView(u8"\u00E5 \u00E4 \u00F6")));	// å ä ö
 			Utf8.AddRange(Batch.NameMember("CJK"),		BuildLeafRange(Scratch, MakeArrayView(u8"\u3300 \uFE30")));			// ㌀ ︰
 			Utf8.AddRange(Batch.NameMember("Symbols"),	BuildLeafRange(Scratch, MakeArrayView(u8"\u2665 \u01F34C")));		// ♥ 🍌
 
@@ -1119,7 +1119,7 @@ TEST_CASE_NAMED(FPlainPropsReadWriteTest, "System::Core::Serialization::PlainPro
 			CHECK(EqualItems(It1.GrabRange().AsLeaves().AsUtf8(), TConstArrayView<char8_t>()));
 			CHECK(EqualItems(It1.GrabRange().AsLeaves().AsUtf8(), u8""));
 			CHECK(EqualItems(It1.GrabRange().AsLeaves().AsUtf8(), MakeArrayView(u8"\"\\ \x01 \x1f\" \"\b \f \n \r \t \"\\").LeftChop(1)));
-			CHECK(EqualItems(It1.GrabRange().AsLeaves().AsUtf8(), u8"\xE5 \xE4 \xF6"));			// å ä ö
+			CHECK(EqualItems(It1.GrabRange().AsLeaves().AsUtf8(), u8"\u00E5 \u00E4 \u00F6"));	// å ä ö
 			CHECK(EqualItems(It1.GrabRange().AsLeaves().AsUtf8(), u8"\u3300 \uFE30"));			// ㌀ ︰
 			CHECK(EqualItems(It1.GrabRange().AsLeaves().AsUtf8(), u8"\u2665 \u01F34C"));		// ♥ 🍌
 
