PlainProps different friend access approach

#rnx

diff --git a/Engine/Plugins/Experimental/PlainProps/Source/Private/PlainPropsInternalRead.h b/Engine/Plugins/Experimental/PlainProps/Source/Private/PlainPropsInternalRead.h
--- a/Engine/Plugins/Experimental/PlainProps/Source/Private/PlainPropsInternalRead.h
+++ b/Engine/Plugins/Experimental/PlainProps/Source/Private/PlainPropsInternalRead.h
@@ -9,6 +9,9 @@
 
 namespace PlainProps
 {
+	
+inline FRangeSchema GetSchema(FRangeView Range) { return Range.Schema; }
+inline FMemoryView GetValues(FRangeView Range) { return Range.Values; }
 
 template<typename MemberType, typename IndexType>
 TConstArrayView<MemberType> GrabInnerRangeTypes(TConstArrayView<MemberType> InnerRangeTypes, IndexType& InOutIdx)
diff --git a/Engine/Plugins/Experimental/PlainProps/Source/Private/PlainPropsLoad.cpp b/Engine/Plugins/Experimental/PlainProps/Source/Private/PlainPropsLoad.cpp
--- a/Engine/Plugins/Experimental/PlainProps/Source/Private/PlainPropsLoad.cpp
+++ b/Engine/Plugins/Experimental/PlainProps/Source/Private/PlainPropsLoad.cpp
@@ -1258,16 +1258,17 @@ FRangeLoadView FMemberLoader::GrabRange()
 	const FRangeView In = Reader.GrabRange();
 
 	// Replace ReadId with LoadId
-	FOptionalSchemaId InnermostId = In.Schema.InnermostSchema;
-	if (LoadIdIt && InnermostId && GetInnermostType(In.Schema).IsStruct())
+	FRangeSchema InSchema = GetSchema(In);
+	FOptionalSchemaId InnermostId = InSchema.InnermostSchema;
+	if (LoadIdIt && InnermostId && GetInnermostType(InSchema).IsStruct())
 	{
 		check(LoadIdIt->Idx > InnermostId.Get().Idx);
 		check(InnermostId == Batch.GetReadId(*LoadIdIt));
 		InnermostId = *LoadIdIt++;
 	}
 
-	FRangeLoadSchema OutSchema = { In.Schema.ItemType, InnermostId, In.Schema.NestedItemTypes, Batch };
-	return { OutSchema, In.NumItems, In.Values };
+	FRangeLoadSchema OutSchema = { InSchema.ItemType, InnermostId, InSchema.NestedItemTypes, Batch };
+	return { OutSchema, In.Num(), GetValues(In) };
 }
 
 FStructLoadView FMemberLoader::GrabStruct()
diff --git a/Engine/Plugins/Experimental/PlainProps/Source/Private/PlainPropsPrint.cpp b/Engine/Plugins/Experimental/PlainProps/Source/Private/PlainPropsPrint.cpp
--- a/Engine/Plugins/Experimental/PlainProps/Source/Private/PlainPropsPrint.cpp
+++ b/Engine/Plugins/Experimental/PlainProps/Source/Private/PlainPropsPrint.cpp
@@ -1113,7 +1113,7 @@ void FMemberPrinter::PrintMembersInternal(FStructType StructType, FStructView St
 
 void FMemberPrinter::PrintRangeInternal(FRangeType RangeType, const FRangeView& RangeView)
 {
-	const FRangeSchema& Schema = RangeView.Schema;
+	FRangeSchema Schema = GetSchema(RangeView);
 	if (RangeView.Num() > 0)
 	{
 		PrintSchemaComment(RangeType, Schema);
@@ -1141,7 +1141,7 @@ void FMemberPrinter::PrintRangeInternal(FRangeType RangeType, const FRangeView&
 
 bool FMemberPrinter::IsUnicodeString(const FRangeView& RangeView)
 {
-	FMemberType Type = RangeView.Schema.ItemType;
+	FMemberType Type = GetSchema(RangeView).ItemType;
 	return RangeView.Num() > 0 && Type.IsLeaf() && Type.AsLeaf().Type == ELeafType::Unicode;
 }
 
@@ -1178,7 +1178,7 @@ void FMemberPrinter::PrintUnicodeRangeAsLeaf(FOptionalMemberId Id, FRangeType Ra
 	}
 
 	const FLeafRangeView LeafRange = RangeView.AsLeaves();
-	const FUnpackedLeafType Leaf = RangeView.Schema.ItemType.AsLeaf();
+	const FUnpackedLeafType Leaf = GetSchema(RangeView).ItemType.AsLeaf();
 
 	switch (Leaf.Width)
 	{
@@ -1188,7 +1188,7 @@ void FMemberPrinter::PrintUnicodeRangeAsLeaf(FOptionalMemberId Id, FRangeType Ra
 		case ELeafWidth::B64:	check(false);															break;
 	};
 	
-	PrintSchemaComment(FRangeType(RangeType), RangeView.Schema);
+	PrintSchemaComment(FRangeType(RangeType), GetSchema(RangeView));
 }
 
 ///////////////////////////////////////////////////////////////////////////////
diff --git a/Engine/Plugins/Experimental/PlainProps/Source/Public/PlainPropsRead.h b/Engine/Plugins/Experimental/PlainProps/Source/Public/PlainPropsRead.h
--- a/Engine/Plugins/Experimental/PlainProps/Source/Public/PlainPropsRead.h
+++ b/Engine/Plugins/Experimental/PlainProps/Source/Public/PlainPropsRead.h
@@ -13,8 +13,6 @@ namespace PlainProps
 {
 
 class FLeafRangeLoadView;
-class FMemberLoader;
-class FMemberPrinter;
 struct FEnumSchema;
 struct FSchemaBatch;
 struct FStructSchema;
@@ -30,7 +28,7 @@ class FSchemaBatchId
 
 // @param Schemas must outlive read batch
 PLAINPROPS_API const FSchemaBatch*		ValidateSchemas(FMemoryView Schemas);
-PLAINPROPS_API FSchemaBatchId				MountReadSchemas(const FSchemaBatch* Schemas);
+PLAINPROPS_API FSchemaBatchId			MountReadSchemas(const FSchemaBatch* Schemas);
 PLAINPROPS_API const FSchemaBatch*		UnmountReadSchemas(FSchemaBatchId Batch);
 
 PLAINPROPS_API uint32					NumStructSchemas(FSchemaBatchId Batch);
@@ -316,8 +314,8 @@ public:
 	PLAINPROPS_API FNestedRangeView	AsRanges() const;		// @pre IsNestedRange()
 
 private:
-	friend FMemberLoader;
-	friend FMemberPrinter;
+	friend FRangeSchema GetSchema(FRangeView Range);
+	friend FMemoryView GetValues(FRangeView Range);
 	friend FMemberReader;
 	friend FNestedRangeIterator;
 
