PlainProps misc minor changes

#rnx

diff --git a/Engine/Plugins/Experimental/PlainProps/Source/Private/PlainPropsInternalRead.h b/Engine/Plugins/Experimental/PlainProps/Source/Private/PlainPropsInternalRead.h
--- a/Engine/Plugins/Experimental/PlainProps/Source/Private/PlainPropsInternalRead.h
+++ b/Engine/Plugins/Experimental/PlainProps/Source/Private/PlainPropsInternalRead.h
@@ -10,8 +10,10 @@
 namespace PlainProps
 {
 	
-inline FRangeSchema GetSchema(FRangeView Range) { return Range.Schema; }
-inline FMemoryView GetValues(FRangeView Range) { return Range.Values; }
+inline FUnpackedLeafType	GetType(FLeafRangeView Range)		{ return { Range.Type, Range.Width }; }
+inline const uint8*			GetValues(FLeafRangeView Range)		{ return Range.Values; }
+inline FRangeSchema			GetSchema(FRangeView Range)			{ return Range.Schema; }
+inline FMemoryView			GetValues(FRangeView Range)			{ return Range.Values; }
 
 template<typename MemberType, typename IndexType>
 TConstArrayView<MemberType> GrabInnerRangeTypes(TConstArrayView<MemberType> InnerRangeTypes, IndexType& InOutIdx)
@@ -141,4 +143,13 @@ inline FMemoryView GetSchemaData(const FSchemaBatch& Batch)
 	return FMemoryView(BatchPtr + Offsets[0], Offsets.Last() + LastSchemaSize - Offsets[0]);
 }
 
+//////////////////////////////////////////////////////////////////////////
+
+template<typename T>
+auto GetEnd(T&& Container)
+{
+	return Container.GetData() + Container.Num();
+}
+
+
 } // namespace PlainProps
diff --git a/Engine/Plugins/Experimental/PlainProps/Source/Public/PlainPropsBuild.h b/Engine/Plugins/Experimental/PlainProps/Source/Public/PlainPropsBuild.h
--- a/Engine/Plugins/Experimental/PlainProps/Source/Public/PlainPropsBuild.h
+++ b/Engine/Plugins/Experimental/PlainProps/Source/Public/PlainPropsBuild.h
@@ -388,7 +388,13 @@ private:
 
 } // namespace PlainProps
 
+
+inline void* operator new(std::size_t Size, PlainProps::FScratchAllocator& Scratch)
+{
+	return Scratch.Allocate(Size, 8);
+}
+
 inline void* operator new(std::size_t Size, std::align_val_t Align, PlainProps::FScratchAllocator& Scratch)
 {
 	return Scratch.Allocate(Size, static_cast<uint32>(Align));
-}
\ No newline at end of file
+}
diff --git a/Engine/Plugins/Experimental/PlainProps/Source/Public/PlainPropsRead.h b/Engine/Plugins/Experimental/PlainProps/Source/Public/PlainPropsRead.h
--- a/Engine/Plugins/Experimental/PlainProps/Source/Public/PlainPropsRead.h
+++ b/Engine/Plugins/Experimental/PlainProps/Source/Public/PlainPropsRead.h
@@ -518,6 +518,9 @@ private:
 	FOptionalEnumSchemaId			Enum;
 	uint64							NumItems;
 	const uint8*					Values;
+
+	friend FUnpackedLeafType GetType(FLeafRangeView Range);
+	friend const uint8* GetValues(FLeafRangeView Range);
 };
 
 //////////////////////////////////////////////////////////////////////////
diff --git a/Engine/Plugins/Experimental/PlainProps/Source/Public/PlainPropsTypes.h b/Engine/Plugins/Experimental/PlainProps/Source/Public/PlainPropsTypes.h
--- a/Engine/Plugins/Experimental/PlainProps/Source/Public/PlainPropsTypes.h
+++ b/Engine/Plugins/Experimental/PlainProps/Source/Public/PlainPropsTypes.h
@@ -332,6 +332,11 @@ public:
 		return FromIdx<T>(Idx);
 	}
 
+	T GetOr(T Fallback) const
+	{
+		return *this ? FromIdx<T>(Idx) : Fallback;
+	}
+
 	friend bool operator==(TOptionalId A, TOptionalId B) { return A.Idx == B.Idx;}
 	friend bool operator!=(TOptionalId A, TOptionalId B) { return A.Idx != B.Idx;}
 };
@@ -538,7 +543,7 @@ public:
 class FDebugIds
 {
 public:
-	PLAINPROPS_API explicit	FDebugIds(const FIds& InIds) : Ids(InIds) {}
+	explicit FDebugIds(const FIds& InIds) : Ids(InIds) {}
 
 	PLAINPROPS_API FString			Print(FNameId Name) const;
 	PLAINPROPS_API FString			Print(FMemberId Name) const;
