PlainProps renamed a few things

#rnx

diff --git a/Engine/Plugins/Experimental/PlainProps/Source/Private/PlainPropsBind.cpp b/Engine/Plugins/Experimental/PlainProps/Source/Private/PlainPropsBind.cpp
--- a/Engine/Plugins/Experimental/PlainProps/Source/Private/PlainPropsBind.cpp
+++ b/Engine/Plugins/Experimental/PlainProps/Source/Private/PlainPropsBind.cpp
@@ -28,7 +28,7 @@ uint32 FSchemaBinding::CalculateSize() const
 {
 	uint32 Out = sizeof(FSchemaBinding) + (NumMembers + NumInnerRanges) * sizeof(FMemberBindType);
 	Out = Align(Out + NumMembers * sizeof(uint32), sizeof(uint32));
-	Out = Align(Out + NumInnerSchemas * sizeof(FInnerId), sizeof(FInnerId));
+	Out = Align(Out + NumInnerIds * sizeof(FInnerId), sizeof(FInnerId));
 	Out = Align(Out + NumInnerRanges * sizeof(FRangeBinding), sizeof(FRangeBinding));
 	return Out;
 }
@@ -116,11 +116,11 @@ void FMemberVisitor::SkipMember()
 	if (Type.IsRange())
 	{
 		FMemberBindTypeRange InnerTypes = GrabInnerTypes();
-		InnerSchemaIdx += HasSchema(InnerTypes.Last());
+		InnerIdIdx += HasSchema(InnerTypes.Last());
 	}
 	else
 	{
-		InnerSchemaIdx += HasSchema(Type);
+		InnerIdIdx += HasSchema(Type);
 	}
 	
 	++MemberIdx;
@@ -131,14 +131,14 @@ FBindId FMemberVisitor::GrabSuper()
 	check(MemberIdx == 0);
 	checkSlow(Schema.Members[0].AsByte() == SuperStructType.AsByte());
 	MemberIdx = 1;
-	InnerSchemaIdx = 1;
-	return Schema.GetInnerSchemas()[0].AsStructBindId();
+	InnerIdIdx = 1;
+	return Schema.GetInnerIds()[0].AsStructBindId();
 }
 
 FInnerId FMemberVisitor::GrabInnerSchema()
 {
-	check(InnerSchemaIdx < Schema.NumInnerSchemas);
-	return Schema.GetInnerSchemas()[InnerSchemaIdx++];
+	check(InnerIdIdx < Schema.NumInnerIds);
+	return Schema.GetInnerIds()[InnerIdIdx++];
 }
 
 ////////////////////////////////////////////////////////////////////////////////////////////////
@@ -442,7 +442,7 @@ struct FMemberBinder : FMemberBinderBase
 {
 	FMemberBinder(FSchemaBinding& InSchema)
 	: FMemberBinderBase(InSchema)
-	, InnerSchemaIt(const_cast<FInnerId*>(InSchema.GetInnerSchemas()))
+	, InnerSchemaIt(const_cast<FInnerId*>(InSchema.GetInnerIds()))
 	{}
 	
 	~FMemberBinder()
@@ -487,7 +487,7 @@ static void DeclareOrValidate(FStructDeclarationPtr& Out, FStructSpec Info, TCon
 		check(Out->Version == Info.Version);
 		check(Out->Occupancy == Info.Occupancy);
 		check(Algo::Compare(Out->GetMemberOrder(), Info.MemberNames));
-		check(Out->NumInnerIds == Binding.NumInnerSchemas);
+		check(Out->NumInnerIds == Binding.NumInnerIds);
 		check(Out->NumInnerRanges == Binding.NumInnerRanges);
 		// Todo: check Out->GetTypes() matches Binding.Members
 		IncRef(Out);
@@ -495,7 +495,7 @@ static void DeclareOrValidate(FStructDeclarationPtr& Out, FStructSpec Info, TCon
 	}
 
 	// Make header and allocate
-	FStructDeclaration Header = { 1, Info.Id, Info.Super, Info.Version, Num, Binding.NumInnerRanges, Binding.NumInnerSchemas, Info.Occupancy };
+	FStructDeclaration Header = { 1, Info.Id, Info.Super, Info.Version, Num, Binding.NumInnerRanges, Binding.NumInnerIds, Info.Occupancy };
 	FStructDeclaration* Ptr = new (FMemory::MallocZeroed(Header.CalculateSize())) FStructDeclaration;
 	FMemory::Memcpy(Ptr, &Header, sizeof(Header));
 
@@ -566,7 +566,7 @@ void FSchemaBindings::BindStruct(FBindId Id, TConstArrayView<FMemberBinding> Mem
 {
 	FSchemaBinding* Schema = Bind(Id, Members, Spec.Id);
 
-	// Extract declared struct ids, Binding::GetInnerSchemas() can contain unlowered FBindIds
+	// Extract declared struct ids, Binding::GetInnerIds() can contain unlowered FBindIds
 	TArray<FInnerId, TInlineAllocator<16>> InnerIds;
 	for (FMemberSpecView Member : Spec.MemberTypes)
 	{
@@ -575,7 +575,7 @@ void FSchemaBindings::BindStruct(FBindId Id, TConstArrayView<FMemberBinding> Mem
 			InnerIds.Emplace(Member.InnermostId.Get());
 		}
 	}
-	check(InnerIds.Num() == Schema->NumInnerSchemas);
+	check(InnerIds.Num() == Schema->NumInnerIds);
 	DeclareOrValidate(/* in-out */ At(Spec.Id), Spec, InnerIds, *Schema);
 }
 
diff --git a/Engine/Plugins/Experimental/PlainProps/Source/Private/PlainPropsDiff.cpp b/Engine/Plugins/Experimental/PlainProps/Source/Private/PlainPropsDiff.cpp
--- a/Engine/Plugins/Experimental/PlainProps/Source/Private/PlainPropsDiff.cpp
+++ b/Engine/Plugins/Experimental/PlainProps/Source/Private/PlainPropsDiff.cpp
@@ -255,7 +255,7 @@ static bool DiffStructSchema(const FStructSchema& A, const FStructSchema& B)
 	check(EqualItems(
 		MakeArrayView(A.GetInnerSchemas(), A.NumInnerSchemas),
 		MakeArrayView(B.GetInnerSchemas(), B.NumInnerSchemas)));
-	check(A.GetSuperSchema() == B.GetSuperSchema());
+	check(A.GetSuper() == B.GetSuper());
 	check(A == B);
 #endif
 	return Diff;
diff --git a/Engine/Plugins/Experimental/PlainProps/Source/Private/PlainPropsInternalBind.h b/Engine/Plugins/Experimental/PlainProps/Source/Private/PlainPropsInternalBind.h
--- a/Engine/Plugins/Experimental/PlainProps/Source/Private/PlainPropsInternalBind.h
+++ b/Engine/Plugins/Experimental/PlainProps/Source/Private/PlainPropsInternalBind.h
@@ -31,7 +31,7 @@ protected: // for unit tests
 	const uint16				NumMembers;
 	uint16						MemberIdx = 0;
 	uint16						InnerRangeIdx = 0;		// Types of [nested] ranges
-	uint16						InnerSchemaIdx = 0;		// Types of static structs and enums
+	uint16						InnerIdIdx = 0;			// Types of static structs and enums
 
 	using FMemberBindTypeRange = TConstArrayView<FMemberBindType>;
 
@@ -59,7 +59,7 @@ struct FMemberBinderBase
 	{
 		check(MemberIt == Schema.GetInnerRangeTypes());
 		check(Align(RangeTypeIt, alignof(uint32)) == (const void*)Schema.GetOffsets());
-		check(OffsetIt == (const void*)Schema.GetInnerSchemas());
+		check(OffsetIt == (const void*)Schema.GetInnerIds());
 		check(Schema.NumInnerRanges == RangeBindingIt - Schema.GetRangeBindings());
 	}
 
diff --git a/Engine/Plugins/Experimental/PlainProps/Source/Private/PlainPropsInternalFormat.h b/Engine/Plugins/Experimental/PlainProps/Source/Private/PlainPropsInternalFormat.h
--- a/Engine/Plugins/Experimental/PlainProps/Source/Private/PlainPropsInternalFormat.h
+++ b/Engine/Plugins/Experimental/PlainProps/Source/Private/PlainPropsInternalFormat.h
@@ -78,7 +78,7 @@ struct FStructSchema
 	TConstArrayView<FMemberId>		GetMemberNames() const		{ return MakeArrayView(GetMemberNames(Footer, NumMembers, NumRangeTypes), NumNames()); }
 	TArrayView<FMemberId>			EditMemberNames() 			{ return MakeArrayView(const_cast<FMemberId*>(GetMemberNames(Footer, NumMembers, NumRangeTypes)), NumNames()); }
 	const FSchemaId*				GetInnerSchemas() const		{ return GetInnerSchemas(Footer, NumMembers, NumRangeTypes, NumNames()); }
-	FOptionalStructSchemaId			GetSuperSchema() const		{ return Inheritance != ESuper::No ? ToOptional(static_cast<FStructSchemaId>(*GetInnerSchemas())) : NoId; }
+	FOptionalStructSchemaId			GetSuper() const			{ return Inheritance != ESuper::No ? ToOptional(static_cast<FStructSchemaId>(*GetInnerSchemas())) : NoId; }
 
 	static const FMemberType*		GetMemberTypes(const FMemberType* Footer)
 	{
diff --git a/Engine/Plugins/Experimental/PlainProps/Source/Private/PlainPropsLoad.cpp b/Engine/Plugins/Experimental/PlainProps/Source/Private/PlainPropsLoad.cpp
--- a/Engine/Plugins/Experimental/PlainProps/Source/Private/PlainPropsLoad.cpp
+++ b/Engine/Plugins/Experimental/PlainProps/Source/Private/PlainPropsLoad.cpp
@@ -28,7 +28,7 @@ struct FSchemaLoadPlan
 
 	TConstArrayView<FMemberBindType>	GetMembers() const			{ return MakeArrayView(Clone.Members, Clone.NumMembers); }
 	const uint32*						GetOffsets() const			{ return Clone.GetOffsets(); }
-	TConstArrayView<FStructSchemaId>	GetInnerSchemas() const		{ return MakeArrayView(reinterpret_cast<const FStructSchemaId*>(Clone.GetInnerSchemas()), Clone.NumInnerSchemas); }
+	TConstArrayView<FStructSchemaId>	GetInnerSchemas() const		{ return MakeArrayView(reinterpret_cast<const FStructSchemaId*>(Clone.GetInnerIds()), Clone.NumInnerIds); }
 	TConstArrayView<FMemberBindType>	GetInnerRangeTypes() const	{ return MakeArrayView(Clone.GetInnerRangeTypes(), Clone.NumInnerRanges); }
 	const FRangeBinding*				GetRangeBindings() const	{ return Clone.GetRangeBindings(); }
 
@@ -163,8 +163,8 @@ static bool HasDifferentSupers(const FStructSchema& From, const FSchemaBinding&
 	}
 	else if (To.HasSuper())
 	{
-		FStructId FromSuper = ToStructIds[From.GetSuperSchema().Get().Idx];
-		FStructId ToSuper = To.GetInnerSchemas()[0].AsStruct();
+		FStructId FromSuper = ToStructIds[From.GetSuper().Get().Idx];
+		FStructId ToSuper = To.GetInnerIds()[0].AsStruct();
 		return FromSuper == ToSuper;
 	}
 
@@ -407,8 +407,8 @@ static void CloneBindingWithReplacedStructIds(const FSchemaId* FromIds, const FS
 
 	// Replace inner bind ids with batch load ids
 	const FStructSchemaId* ReadIdIt = static_cast<const FStructSchemaId*>(FromIds);
-	FInnerId* InnerIt = const_cast<FInnerId*>(Schema->GetInnerSchemas());
-	for (FInnerId& Inner : MakeArrayView(InnerIt, To.NumInnerSchemas))
+	FInnerId* InnerIt = const_cast<FInnerId*>(Schema->GetInnerIds());
+	for (FInnerId& Inner : MakeArrayView(InnerIt, To.NumInnerIds))
 	{
 		FBindId MemcopiedBindId = Inner.AsStructBindId();
 		reinterpret_cast<FStructSchemaId&>(Inner) = LoadIds.BindLoadId(*ReadIdIt++, MemcopiedBindId);
@@ -426,7 +426,7 @@ template<class LoadIdBinder>
 	else
 	{
 		check(From.NumMembers == To.NumMembers);
-		check(From.NumInnerSchemas == To.NumInnerSchemas);
+		check(From.NumInnerSchemas == To.NumInnerIds);
 		check(From.NumRangeTypes == To.NumInnerRanges);
 
 		if (From.NumInnerSchemas > 0)
diff --git a/Engine/Plugins/Experimental/PlainProps/Source/Public/PlainPropsBind.h b/Engine/Plugins/Experimental/PlainProps/Source/Public/PlainPropsBind.h
--- a/Engine/Plugins/Experimental/PlainProps/Source/Public/PlainPropsBind.h
+++ b/Engine/Plugins/Experimental/PlainProps/Source/Public/PlainPropsBind.h
@@ -139,16 +139,16 @@ struct alignas(/*FRangeBinding*/ 8) FSchemaBinding
 {
 	FDeclId					DeclId;
 	uint16					NumMembers;
-	uint16					NumInnerSchemas;
+	uint16					NumInnerIds;
 	uint16					NumInnerRanges;
 	FMemberBindType			Members[0];
 
 	const FMemberBindType*	GetInnerRangeTypes() const	{ return Members + NumMembers; }
 	const uint32*			GetOffsets() const			{ return AlignPtr<uint32>(GetInnerRangeTypes() + NumInnerRanges); }
-	const FInnerId*			GetInnerSchemas() const		{ return AlignPtr<FInnerId>(GetOffsets() + NumMembers); }
-	const FRangeBinding*	GetRangeBindings() const	{ return AlignPtr<FRangeBinding>(GetInnerSchemas() + NumInnerSchemas); }
+	const FInnerId*			GetInnerIds() const			{ return AlignPtr<FInnerId>(GetOffsets() + NumMembers); }
+	const FRangeBinding*	GetRangeBindings() const	{ return AlignPtr<FRangeBinding>(GetInnerIds() + NumInnerIds); }
 	uint32					CalculateSize() const;
-	bool					HasSuper() const			{ return NumInnerSchemas > 0 && Members[0].IsStruct() && Members[0].AsStruct().IsSuper; }
+	bool					HasSuper() const			{ return NumInnerIds > 0 && Members[0].IsStruct() && Members[0].AsStruct().IsSuper; }
 };
 
 ////////////////////////////////////////////////////////////////////////////////////////////////
diff --git a/Engine/Plugins/Experimental/PlainProps/Source/Tests/PlainPropsTests.cpp b/Engine/Plugins/Experimental/PlainProps/Source/Tests/PlainPropsTests.cpp
--- a/Engine/Plugins/Experimental/PlainProps/Source/Tests/PlainPropsTests.cpp
+++ b/Engine/Plugins/Experimental/PlainProps/Source/Tests/PlainPropsTests.cpp
@@ -1603,8 +1603,8 @@ TEST_CASE_NAMED(FPlainPropsReadWriteTest, "System::Core::Serialization::PlainPro
 			CHECK(Names[Schema0.Type.Name] == "B6");
 			CHECK(Names[Schema1.Type.Name] == "B5");
 			CHECK(Names[Schema2.Type.Name] == "B6");
-			CHECK(Schema0.GetSuperSchema() == Objects[1].Schema.Id);
-			CHECK(Schema2.GetSuperSchema() == Objects[1].Schema.Id);
+			CHECK(Schema0.GetSuper() == Objects[1].Schema.Id);
+			CHECK(Schema2.GetSuper() == Objects[1].Schema.Id);
 
 			CHECK(Names[FMemberReader(Objects[0]).GrabStruct().Schema.Resolve().Type.Name]  == "B5");
 			CHECK(Names[FMemberReader(Objects[1]).GrabStruct().Schema.Resolve().Type.Name]  == "B4");
@@ -1612,12 +1612,12 @@ TEST_CASE_NAMED(FPlainPropsReadWriteTest, "System::Core::Serialization::PlainPro
 
 			FSchemaBatchId BatchId = Objects[0].Schema.Batch;
 			const FStructSchema& B6 = Objects[0].Schema.Resolve();
-			const FStructSchema& B5 = ResolveStructSchema(BatchId, B6.GetSuperSchema().Get());
-			const FStructSchema& B4 = ResolveStructSchema(BatchId, B5.GetSuperSchema().Get());
-			const FStructSchema& B3 = ResolveStructSchema(BatchId, B4.GetSuperSchema().Get());
-			const FStructSchema& B2 = ResolveStructSchema(BatchId, B3.GetSuperSchema().Get());
-			const FStructSchema& B1 = ResolveStructSchema(BatchId, B2.GetSuperSchema().Get());
-			const FStructSchema& B0 = ResolveStructSchema(BatchId, B1.GetSuperSchema().Get());
+			const FStructSchema& B5 = ResolveStructSchema(BatchId, B6.GetSuper().Get());
+			const FStructSchema& B4 = ResolveStructSchema(BatchId, B5.GetSuper().Get());
+			const FStructSchema& B3 = ResolveStructSchema(BatchId, B4.GetSuper().Get());
+			const FStructSchema& B2 = ResolveStructSchema(BatchId, B3.GetSuper().Get());
+			const FStructSchema& B1 = ResolveStructSchema(BatchId, B2.GetSuper().Get());
+			const FStructSchema& B0 = ResolveStructSchema(BatchId, B1.GetSuper().Get());
 			// Super usage									A	B	C 	Decl
 			CHECK(B0.Inheritance == ESuper::No);		//	-	-	-	-
 			CHECK(B1.Inheritance == ESuper::Unused);	//	0	0	0	B0
