PlainProps renamed ERangeSizeType -> ESizeType

#rnx

diff --git a/Engine/Plugins/Experimental/PlainProps/Source/Private/PlainPropsBind.cpp b/Engine/Plugins/Experimental/PlainProps/Source/Private/PlainPropsBind.cpp
--- a/Engine/Plugins/Experimental/PlainProps/Source/Private/PlainPropsBind.cpp
+++ b/Engine/Plugins/Experimental/PlainProps/Source/Private/PlainPropsBind.cpp
@@ -100,7 +100,7 @@ TConstArrayView<FMemberBindType> FMemberVisitor::GrabInnerTypes()
 
 FRangeMemberBinding FMemberVisitor::GrabRange()
 {
-	ERangeSizeType MaxSize = PeekType().AsRange().MaxSize;
+	ESizeType MaxSize = PeekType().AsRange().MaxSize;
 	const FRangeBinding* RangeBindings = Schema.GetRangeBindings() + InnerRangeIdx;
 	check(MaxSize == RangeBindings[0].GetSizeType());
 	FMemberBindTypeRange InnerTypes = GrabInnerTypes();
@@ -143,14 +143,14 @@ FInnerId FMemberVisitor::GrabInnerSchema()
 
 ////////////////////////////////////////////////////////////////////////////////////////////////
 
-FRangeBinding::FRangeBinding(const IItemRangeBinding& Binding, ERangeSizeType SizeType)
+FRangeBinding::FRangeBinding(const IItemRangeBinding& Binding, ESizeType SizeType)
 : Handle(uint64(&Binding) | uint8(SizeType))
 {
 	check(&Binding == &AsItemBinding());
 	check(SizeType == GetSizeType());
 }
 
-FRangeBinding::FRangeBinding(const ILeafRangeBinding& Binding, ERangeSizeType SizeType)
+FRangeBinding::FRangeBinding(const ILeafRangeBinding& Binding, ESizeType SizeType)
 : Handle(uint64(&Binding) | uint8(SizeType) | LeafMask)
 {
 	check(&Binding == &AsLeafBinding());
diff --git a/Engine/Plugins/Experimental/PlainProps/Source/Private/PlainPropsBuild.cpp b/Engine/Plugins/Experimental/PlainProps/Source/Private/PlainPropsBuild.cpp
--- a/Engine/Plugins/Experimental/PlainProps/Source/Private/PlainPropsBuild.cpp
+++ b/Engine/Plugins/Experimental/PlainProps/Source/Private/PlainPropsBuild.cpp
@@ -70,7 +70,7 @@ FBuiltRange* FBuiltRange::Create(FScratchAllocator& Scratch, uint64 NumItems, SI
 
 //////////////////////////////////////////////////////////////////////////
 
-FMemberSchema MakeNestedRangeSchema(FScratchAllocator& Scratch, ERangeSizeType SizeType, FMemberSchema InnerRangeSchema)
+FMemberSchema MakeNestedRangeSchema(FScratchAllocator& Scratch, ESizeType SizeType, FMemberSchema InnerRangeSchema)
 {
 	check(InnerRangeSchema.NumInnerRanges > 0);
 	uint16 NumInnerRanges = IntCastChecked<uint16>(1 + InnerRangeSchema.NumInnerRanges);
@@ -275,7 +275,7 @@ FNestedRangeBuilder::~FNestedRangeBuilder()
 	checkf(Ranges.IsEmpty(), TEXT("Half-built range, forgot to call BuildAndReset() before destruction?"));
 }
 
-FTypedRange FNestedRangeBuilder::BuildAndReset(FScratchAllocator& Scratch, ERangeSizeType SizeType)
+FTypedRange FNestedRangeBuilder::BuildAndReset(FScratchAllocator& Scratch, ESizeType SizeType)
 {
 	FBuiltRange* Out = nullptr;
 	if (int64 Num = Ranges.Num())
diff --git a/Engine/Plugins/Experimental/PlainProps/Source/Private/PlainPropsBuildSchema.cpp b/Engine/Plugins/Experimental/PlainProps/Source/Private/PlainPropsBuildSchema.cpp
--- a/Engine/Plugins/Experimental/PlainProps/Source/Private/PlainPropsBuildSchema.cpp
+++ b/Engine/Plugins/Experimental/PlainProps/Source/Private/PlainPropsBuildSchema.cpp
@@ -55,7 +55,7 @@ struct FStructSchemaBuilder
 	TBitArray<>									NotedMembers;
 
 	void										NoteMembersRecursively(const FBuiltStruct& Struct);
-	void										NoteRangeRecursively(ERangeSizeType NumType, TConstArrayView<FMemberType> Types, void* InnermostSchemaBuilder, const FBuiltRange* Range, FMemberId Member);
+	void										NoteRangeRecursively(ESizeType NumType, TConstArrayView<FMemberType> Types, void* InnermostSchemaBuilder, const FBuiltRange* Range, FMemberId Member);
 	FBuiltStructSchema							Build() const;
 };
 
@@ -334,7 +334,7 @@ static void NoteEnumRange(FEnumSchemaBuilder& Out, FLeafType Leaf, const FBuiltR
 	}
 }
 
-static void NoteEmptyRange(ERangeSizeType NumType, TConstArrayView<FMemberType> Types, void* InnermostSchema)
+static void NoteEmptyRange(ESizeType NumType, TConstArrayView<FMemberType> Types, void* InnermostSchema)
 {
 	FMemberType InnermostType = Types.Last();
 	if (IsEnum(InnermostType))
@@ -343,7 +343,7 @@ static void NoteEmptyRange(ERangeSizeType NumType, TConstArrayView<FMemberType>
 	}
 }
 
-void FStructSchemaBuilder::NoteRangeRecursively(ERangeSizeType NumType, TConstArrayView<FMemberType> Types, void* InnermostSchema, const FBuiltRange* Range, FMemberId Member)
+void FStructSchemaBuilder::NoteRangeRecursively(ESizeType NumType, TConstArrayView<FMemberType> Types, void* InnermostSchema, const FBuiltRange* Range, FMemberId Member)
 {
 	if (Range == nullptr)
 	{
diff --git a/Engine/Plugins/Experimental/PlainProps/Source/Private/PlainPropsDeclare.cpp b/Engine/Plugins/Experimental/PlainProps/Source/Private/PlainPropsDeclare.cpp
--- a/Engine/Plugins/Experimental/PlainProps/Source/Private/PlainPropsDeclare.cpp
+++ b/Engine/Plugins/Experimental/PlainProps/Source/Private/PlainPropsDeclare.cpp
@@ -187,7 +187,7 @@ FStructDeclarationPtr Declare(FStructSpec In)
 		if (int32 NumRanges = Member.Ranges.Num())
 		{
 			*TypeIt++ = FMemberType(Member.Ranges[0]);
-			for (ERangeSizeType Range : Member.Ranges.RightChop(1))
+			for (ESizeType Range : Member.Ranges.RightChop(1))
 			{
 				*RangeIt++ = FMemberType(Range);
 			}
@@ -207,7 +207,7 @@ FStructDeclarationPtr Declare(FStructSpec In)
 
 ////////////////////////////////////////////////////////////////////////////////////////////////
 
-FMemberSpec::FMemberSpec(ERangeSizeType MaxSize, FMemberSpec Inner)
+FMemberSpec::FMemberSpec(ESizeType MaxSize, FMemberSpec Inner)
 : FMemberSpec(Inner)
 {
 	RangeWrap(MaxSize);
diff --git a/Engine/Plugins/Experimental/PlainProps/Source/Private/PlainPropsDiff.cpp b/Engine/Plugins/Experimental/PlainProps/Source/Private/PlainPropsDiff.cpp
--- a/Engine/Plugins/Experimental/PlainProps/Source/Private/PlainPropsDiff.cpp
+++ b/Engine/Plugins/Experimental/PlainProps/Source/Private/PlainPropsDiff.cpp
@@ -414,7 +414,7 @@ static bool DiffLeaves(FUnpackedLeafType Leaf, FLeafRangeView A, FLeafRangeView
 }
 
 static bool DiffMembers(FStructView A, FStructView B, FReadDiffPath& Out);
-static bool DiffRange(ERangeSizeType NumType, FRangeView A, FRangeView B, FReadDiffPath& Out);
+static bool DiffRange(ESizeType NumType, FRangeView A, FRangeView B, FReadDiffPath& Out);
 
 static bool DiffStructs(FStructType Struct, FStructRangeView A, FStructRangeView B, FReadDiffPath& Out)
 {
@@ -446,7 +446,7 @@ static bool DiffStructs(FStructType Struct, FStructRangeView A, FStructRangeView
 	return false;
 }
 
-static bool DiffRanges(ERangeSizeType NumType, FNestedRangeView A, FNestedRangeView B, FReadDiffPath& Out)
+static bool DiffRanges(ESizeType NumType, FNestedRangeView A, FNestedRangeView B, FReadDiffPath& Out)
 {
 	uint64 NumA = A.Num();
 	uint64 NumB = B.Num();
@@ -487,7 +487,7 @@ static bool DiffLeaves(FUnpackedLeafType Leaf, FLeafRangeView A, FLeafRangeView
 	return false;
 }
 
-static bool DiffRange(ERangeSizeType NumType, FRangeView A, FRangeView B, FReadDiffPath& Out)
+static bool DiffRange(ESizeType NumType, FRangeView A, FRangeView B, FReadDiffPath& Out)
 {
 	bool bDiff = A.GetItemType() != B.GetItemType();
 	if (!bDiff)
diff --git a/Engine/Plugins/Experimental/PlainProps/Source/Private/PlainPropsInternalRead.h b/Engine/Plugins/Experimental/PlainProps/Source/Private/PlainPropsInternalRead.h
--- a/Engine/Plugins/Experimental/PlainProps/Source/Private/PlainPropsInternalRead.h
+++ b/Engine/Plugins/Experimental/PlainProps/Source/Private/PlainPropsInternalRead.h
@@ -26,19 +26,19 @@ TConstArrayView<MemberType> GrabInnerRangeTypes(TConstArrayView<MemberType> Inne
 	return MakeArrayView(&InnerRangeTypes[StartIdx], Idx - StartIdx);
 }
 
-inline uint64 GrabRangeNum(ERangeSizeType MaxSize, FByteReader& ByteIt, FBitCacheReader& BitIt)
+inline uint64 GrabRangeNum(ESizeType MaxSize, FByteReader& ByteIt, FBitCacheReader& BitIt)
 {
 	switch(MaxSize)
 	{
-	case ERangeSizeType::Uni:	return BitIt.GrabNext(/* in-out */ ByteIt) ? 1 : 0;
-	case ERangeSizeType::S8:	return IntCastChecked<uint64>(	ByteIt.Grab<int8>());
-	case ERangeSizeType::U8:	return							ByteIt.Grab<uint8>();
-	case ERangeSizeType::S16:	return IntCastChecked<uint64>(	ByteIt.Grab<int16>());
-	case ERangeSizeType::U16:	return							ByteIt.Grab<uint16>();
-	case ERangeSizeType::S32:	return IntCastChecked<uint64>(	ByteIt.Grab<int32>());
-	case ERangeSizeType::U32:	return							ByteIt.Grab<uint32>();
-	case ERangeSizeType::S64:	return IntCastChecked<uint64>(	ByteIt.Grab<int64>());
-	case ERangeSizeType::U64:	return							ByteIt.Grab<uint64>();
+	case ESizeType::Uni:	return BitIt.GrabNext(/* in-out */ ByteIt) ? 1 : 0;
+	case ESizeType::S8:		return IntCastChecked<uint64>(	ByteIt.Grab<int8>());
+	case ESizeType::U8:		return							ByteIt.Grab<uint8>();
+	case ESizeType::S16:	return IntCastChecked<uint64>(	ByteIt.Grab<int16>());
+	case ESizeType::U16:	return							ByteIt.Grab<uint16>();
+	case ESizeType::S32:	return IntCastChecked<uint64>(	ByteIt.Grab<int32>());
+	case ESizeType::U32:	return							ByteIt.Grab<uint32>();
+	case ESizeType::S64:	return IntCastChecked<uint64>(	ByteIt.Grab<int64>());
+	case ESizeType::U64:	return							ByteIt.Grab<uint64>();
 	}
 	check(false);
 	return 0;
diff --git a/Engine/Plugins/Experimental/PlainProps/Source/Private/PlainPropsLoad.cpp b/Engine/Plugins/Experimental/PlainProps/Source/Private/PlainPropsLoad.cpp
--- a/Engine/Plugins/Experimental/PlainProps/Source/Private/PlainPropsLoad.cpp
+++ b/Engine/Plugins/Experimental/PlainProps/Source/Private/PlainPropsLoad.cpp
@@ -854,7 +854,7 @@ inline void SetBit(uint8& Out, uint8 Idx, bool bValue)
 
 struct FLoadRangePlan
 {
-	ERangeSizeType MaxSize;
+	ESizeType MaxSize;
 	FOptionalStructSchemaId InnermostStruct;
 	TConstArrayView<FMemberBindType> InnerTypes;
 	const FRangeBinding* Bindings = nullptr;
@@ -892,7 +892,7 @@ public:
 		check(Bindings.Num() == InnerTypes.Num());
 			
 		FOptionalStructSchemaId InnermostStruct = InnerTypes.Last().IsStruct() ? static_cast<FOptionalStructSchemaId>(Src.Schema.InnermostId) : NoId;
-		ERangeSizeType Unused = ERangeSizeType::Uni;
+		ESizeType Unused = ESizeType::Uni;
 		FLoadRangePlan Plan = { Unused, InnermostStruct, InnerTypes, Bindings.GetData() };
 		
 		LoadRangePlan(Member, Src.NumItems, Src.Values, Src.Schema.Batch, Plan);
@@ -947,7 +947,7 @@ public:
 	static void LoadRangeValues(uint8* Member, uint64 Num, const IItemRangeBinding& Binding, FMemoryView Values, const FLoadBatch& Batch, SchemaType&& Schema)
 	{
 		FByteReader ValueIt(Values);
-		FBitCacheReader BitIt; // Only used by ranges of ERangeSizeType::Uni ranges
+		FBitCacheReader BitIt; // Only used by ranges of ESizeType::Uni ranges
 		FLoadRangeContext Ctx{.Request = {Member, Num}};
 		
 		while (Ctx.Request.Index < Num)
@@ -1287,7 +1287,7 @@ void LoadRange(void* Dst, FRangeLoadView Src, TConstArrayView<FRangeBinding> Bin
 	FRangeLoader::LoadView(static_cast<uint8*>(Dst), Src, Bindings);
 }
 
-void LoadRange(void* Dst, FByteReader& SrcBytes, FBitCacheReader& SrcBits, ERangeSizeType MaxSize, FRangeLoadSchema Schema, TConstArrayView<FRangeBinding> Bindings)
+void LoadRange(void* Dst, FByteReader& SrcBytes, FBitCacheReader& SrcBits, ESizeType MaxSize, FRangeLoadSchema Schema, TConstArrayView<FRangeBinding> Bindings)
 {
 	if (uint64 Num = GrabRangeNum(MaxSize, SrcBytes, SrcBits))
 	{
diff --git a/Engine/Plugins/Experimental/PlainProps/Source/Private/PlainPropsParse.cpp b/Engine/Plugins/Experimental/PlainProps/Source/Private/PlainPropsParse.cpp
--- a/Engine/Plugins/Experimental/PlainProps/Source/Private/PlainPropsParse.cpp
+++ b/Engine/Plugins/Experimental/PlainProps/Source/Private/PlainPropsParse.cpp
@@ -71,7 +71,7 @@ static constexpr FUnpackedLeafType GLeaves[] = {
 };
 
 template<>
-TOptional<ERangeSizeType> Parse(FUtf8StringView String)
+TOptional<ESizeType> Parse(FUtf8StringView String)
 {
 	if (String.Len() >= 2)
 	{
@@ -80,7 +80,7 @@ TOptional<ERangeSizeType> Parse(FUtf8StringView String)
 		{
 			if (String == GLiterals.Ranges[I])
 			{
-				return ERangeSizeType{ I };
+				return ESizeType{ I };
 			}
 		}
 	}
@@ -649,7 +649,7 @@ private:
 		FScopedTokenReader InnerTokens);
 
 	[[nodiscard]] FTypedRange	ParseUnicodeLeafValueAsRange(
-		ERangeSizeType NumType,
+		ESizeType NumType,
 		FUnpackedLeafType Leaf,
 		FUtf8StringView String);
 
@@ -1225,7 +1225,7 @@ FParsedStructSchema FStructSchemaParser::ParseMembersAndDeclare(
 	return { Id, Version, DeclaredSuper, UsedSuper, Declarations[Id.Idx]->GetMemberOrder(), MoveTemp(MemberSchemas) };
 }
 
-static TOptional<ERangeSizeType> GrabRangeSize(FUtf8StringView& String)
+static TOptional<ESizeType> GrabRangeSize(FUtf8StringView& String)
 {
 	if (String.Len() > 2 && OpenRangeSet.Contains(String[0]))
 	{
@@ -1236,7 +1236,7 @@ static TOptional<ERangeSizeType> GrabRangeSize(FUtf8StringView& String)
 		Remainder.RightChopInline(SizeStr.Len());
 		Skip = FAsciiSet::FindPrefixWith(Remainder, WhitespaceSet | CloseRangeSet);
 		Remainder.RightChopInline(Skip.Len());
-		if (TOptional<ERangeSizeType> MaxSize = Parse<ERangeSizeType>(SizeStr))
+		if (TOptional<ESizeType> MaxSize = Parse<ESizeType>(SizeStr))
 		{
 			String = Remainder;
 			return MaxSize;
@@ -1288,7 +1288,7 @@ bool FStructSchemaParser::ParseMemberSchema(FParsedMemberSchema& Out, FUtf8Strin
 	{
 		FUtf8StringView RangeSizes = String.RightChop(TypeOrLeaf.Len());
 		FMemberType InnermostType = Out.Type;
-		if (TOptional<ERangeSizeType> RangeSize = GrabRangeSize(RangeSizes))
+		if (TOptional<ESizeType> RangeSize = GrabRangeSize(RangeSizes))
 		{
 			Out.Type = FMemberType(RangeSize.GetValue());
 			while ((RangeSize = GrabRangeSize(RangeSizes)))
@@ -1604,7 +1604,7 @@ FBuiltStruct* FMemberParser::BuildAndReset(FDeclId Id)
 template<Arithmetic T>
 [[nodiscard]] static FTypedRange ParseLeafRangeValues(
 	FScratchAllocator& Scratch,
-	ERangeSizeType MaxSize,
+	ESizeType MaxSize,
 	FScopedTokenReader& RangeTokens)
 {
 	TArray<T, TInlineAllocator64<64>> Values;
@@ -1623,7 +1623,7 @@ template<Arithmetic T>
 [[nodiscard]] static FTypedRange ParseEnumRange(
 	FScratchAllocator& Scratch,
 	FEnumId Enum,
-	ERangeSizeType MaxSize,
+	ESizeType MaxSize,
 	FScopedTokenReader& RangeTokens)
 {
 	TArray<T, TInlineAllocator64<64>> Values;
@@ -1643,7 +1643,7 @@ FTypedRange FMemberParser::ParseLeaves(const FMemberSchema& RangeSchema, FScoped
 	const FOptionalEnumId Enum = ToOptionalEnum(RangeSchema.InnerSchema);
 	const FMemberType InnermostType = RangeSchema.GetInnerRangeTypes().Last();
 	const FUnpackedLeafType Leaf = InnermostType.AsLeaf();
-	const ERangeSizeType NumType = RangeSchema.Type.AsRange().MaxSize;
+	const ESizeType NumType = RangeSchema.Type.AsRange().MaxSize;
 
 	FTypedRange TypedRange;
 	switch (Leaf.Type)
@@ -1708,7 +1708,7 @@ FTypedRange FMemberParser::ParseLeaves(const FMemberSchema& RangeSchema, FScoped
 FTypedRange FMemberParser::ParseStructs(const FMemberSchema& RangeSchema, FScopedTokenReader RangeTokens)
 {
 	check(RangeSchema.NumInnerRanges == 1);
-	ERangeSizeType NumType = RangeSchema.Type.AsRange().MaxSize;
+	ESizeType NumType = RangeSchema.Type.AsRange().MaxSize;
 	const FStructType StructType = RangeSchema.InnerRangeType.AsStruct();
 	check(!StructType.IsSuper);
 	check(RangeSchema.InnerSchema || StructType.IsDynamic);
@@ -1764,7 +1764,7 @@ FTypedRange FMemberParser::ParseRanges(const FMemberSchema& RangeSchema, FScoped
 	check(RangeSchema.NumInnerRanges > 1);
 	check(RangeSchema.NestedRangeTypes);
 
-	const ERangeSizeType NumType = RangeSchema.Type.AsRange().MaxSize;
+	const ESizeType NumType = RangeSchema.Type.AsRange().MaxSize;
 	const bool IsDynamic = RangeSchema.GetInnermostType() == DynamicStructType;
 	const FMemberSchema InnerRangeSchema
 	{
@@ -1873,7 +1873,7 @@ static FUtf8StringView ParseString(TArray<UTF8CHAR>& Out, FUtf8StringView String
 }
 
 FTypedRange FMemberParser::ParseUnicodeLeafValueAsRange(
-	ERangeSizeType NumType,
+	ESizeType NumType,
 	FUnpackedLeafType Leaf,
 	FUtf8StringView String)
 {
diff --git a/Engine/Plugins/Experimental/PlainProps/Source/Private/PlainPropsPrint.cpp b/Engine/Plugins/Experimental/PlainProps/Source/Private/PlainPropsPrint.cpp
--- a/Engine/Plugins/Experimental/PlainProps/Source/Private/PlainPropsPrint.cpp
+++ b/Engine/Plugins/Experimental/PlainProps/Source/Private/PlainPropsPrint.cpp
@@ -23,7 +23,7 @@ static constexpr bool PrintWithComments = true;
 
 const FLiterals GLiterals;
 
-FAnsiStringView ToString(ERangeSizeType Width)
+FAnsiStringView ToString(ESizeType Width)
 {
 	return GLiterals.Ranges[(uint8)Width];
 }
diff --git a/Engine/Plugins/Experimental/PlainProps/Source/Private/PlainPropsWrite.cpp b/Engine/Plugins/Experimental/PlainProps/Source/Private/PlainPropsWrite.cpp
--- a/Engine/Plugins/Experimental/PlainProps/Source/Private/PlainPropsWrite.cpp
+++ b/Engine/Plugins/Experimental/PlainProps/Source/Private/PlainPropsWrite.cpp
@@ -736,7 +736,7 @@ private:
 		Tmp.Reset();
 	}
 
-	void WriteRange(ERangeSizeType NumType, TConstArrayView<FMemberType> Types, FOptionalInnerId InnermostSchema, const FBuiltRange* Range)
+	void WriteRange(ESizeType NumType, TConstArrayView<FMemberType> Types, FOptionalInnerId InnermostSchema, const FBuiltRange* Range)
 	{
 		check(Types.Num() > 0);
 		check((Types.Num() > 1) == (Types[0].GetKind() == EMemberKind::Range));
@@ -744,7 +744,7 @@ private:
 
 		// Write Num
 		uint64 Num = Range ? Range->Num : 0;
-		if (NumType == ERangeSizeType::Uni)
+		if (NumType == ESizeType::Uni)
 		{
 			Bits.WriteBit(Num == 1);
 		}
@@ -805,7 +805,7 @@ private:
 		WriteSkippableItems(Structs, [=](FMemberWriter& Out, const FBuiltStruct* Struct) { Out.WriteStruct(StructType, BuiltId, *Struct); });
 	}
 	
-	void WriteRanges(ERangeSizeType NumType, TConstArrayView<FMemberType> Types, FOptionalInnerId InnermostSchema, TConstArrayView64<const FBuiltRange*> Ranges)
+	void WriteRanges(ESizeType NumType, TConstArrayView<FMemberType> Types, FOptionalInnerId InnermostSchema, TConstArrayView64<const FBuiltRange*> Ranges)
 	{	
 		WriteSkippableItems(Ranges, [=](FMemberWriter& Out, const FBuiltRange* Range) { Out.WriteRange(NumType, Types, InnermostSchema, Range); });
 	}
diff --git a/Engine/Plugins/Experimental/PlainProps/Source/Public/PlainPropsBind.h b/Engine/Plugins/Experimental/PlainProps/Source/Public/PlainPropsBind.h
--- a/Engine/Plugins/Experimental/PlainProps/Source/Public/PlainPropsBind.h
+++ b/Engine/Plugins/Experimental/PlainProps/Source/Public/PlainPropsBind.h
@@ -108,7 +108,7 @@ union FMemberBindType
 	constexpr explicit FMemberBindType(FLeafBindType In) : Leaf(In) {}
 	constexpr explicit FMemberBindType(FBitfieldBoolBindType In) : Leaf(In) {}
 	constexpr explicit FMemberBindType(FRangeType In) : Range(In) {}
-	constexpr explicit FMemberBindType(ERangeSizeType MaxSize) : Range({EMemberKind::Range, MaxSize}) {}
+	constexpr explicit FMemberBindType(ESizeType MaxSize) : Range({EMemberKind::Range, MaxSize}) {}
 	constexpr explicit FMemberBindType(FStructType In) : Struct(In) {}
 	
 	bool					IsLeaf() const		{ return Kind == EMemberKind::Leaf; }
@@ -572,13 +572,13 @@ class FRangeBinding
 	uint64						Handle;
 
 public:
-	PLAINPROPS_API FRangeBinding(const IItemRangeBinding& Binding, ERangeSizeType SizeType);
-	PLAINPROPS_API FRangeBinding(const ILeafRangeBinding& Binding, ERangeSizeType SizeType);
+	PLAINPROPS_API FRangeBinding(const IItemRangeBinding& Binding, ESizeType SizeType);
+	PLAINPROPS_API FRangeBinding(const ILeafRangeBinding& Binding, ESizeType SizeType);
 	
 	bool						IsLeafBinding() const					{ return !!(LeafMask & Handle); }
 	const IItemRangeBinding&	AsItemBinding() const					{ check(!IsLeafBinding()); return *reinterpret_cast<IItemRangeBinding*>(Handle & BindMask); }
 	const ILeafRangeBinding&	AsLeafBinding() const					{ check( IsLeafBinding()); return *reinterpret_cast<ILeafRangeBinding*>(Handle & BindMask); }
-	ERangeSizeType				GetSizeType() const						{ return static_cast<ERangeSizeType>(Handle & SizeMask); }
+	ESizeType					GetSizeType() const						{ return static_cast<ESizeType>(Handle & SizeMask); }
 	FConcreteTypenameId			GetBindName() const						{ return IsLeafBinding() ? AsLeafBinding().BindName : AsItemBinding().BindName; }
 	inline bool	operator==(FRangeBinding O) const { return Handle == O.Handle; }
 };
diff --git a/Engine/Plugins/Experimental/PlainProps/Source/Public/PlainPropsBuild.h b/Engine/Plugins/Experimental/PlainProps/Source/Public/PlainPropsBuild.h
--- a/Engine/Plugins/Experimental/PlainProps/Source/Public/PlainPropsBuild.h
+++ b/Engine/Plugins/Experimental/PlainProps/Source/Public/PlainPropsBuild.h
@@ -140,7 +140,7 @@ struct FTypedRange
 };
 
 template<Arithmetic T>
-FMemberSchema MakeLeafRangeSchema(ERangeSizeType MaxSize)
+FMemberSchema MakeLeafRangeSchema(ESizeType MaxSize)
 {
 	return { FMemberType(MaxSize), ReflectArithmetic<T>.Pack(), 1, NoId, nullptr };
 }
@@ -152,7 +152,7 @@ FMemberSchema MakeLeafRangeSchema()
 }
 
 template<Enumeration T>
-FMemberSchema MakeEnumRangeSchema(FEnumId Id, ERangeSizeType MaxSize)
+FMemberSchema MakeEnumRangeSchema(FEnumId Id, ESizeType MaxSize)
 {
 	return { FMemberType(MaxSize), ReflectEnum<T>.Pack(), 1, FInnerId(Id), nullptr };
 }
@@ -164,7 +164,7 @@ FMemberSchema MakeEnumRangeSchema(FEnumId Id)
 }
 
 template<UnsignedIntegral T>
-FMemberSchema MakeEnumRangeSchema(FEnumId Id, ERangeSizeType MaxSize)
+FMemberSchema MakeEnumRangeSchema(FEnumId Id, ESizeType MaxSize)
 {
 	return { FMemberType(MaxSize), FMemberType(ELeafType::Enum, ReflectArithmetic<T>.Width), 1, FInnerId(Id), nullptr };
 }
@@ -174,24 +174,24 @@ inline constexpr FMemberType SuperStructType =		FMemberType(FStructType{EMemberK
 inline constexpr FMemberType DynamicStructType =	FMemberType(FStructType{EMemberKind::Struct, /* IsDynamic */ 1, /* IsSuper */ 0});
 
 
-inline FMemberSchema MakeStructRangeSchema(ERangeSizeType SizeType, FStructId Id)
+inline FMemberSchema MakeStructRangeSchema(ESizeType SizeType, FStructId Id)
 {
 	return { FMemberType(SizeType), DefaultStructType, 1, FInnerId(Id), nullptr };
 }
 
-inline FMemberSchema MakeDynamicStructRangeSchema(ERangeSizeType SizeType)
+inline FMemberSchema MakeDynamicStructRangeSchema(ESizeType SizeType)
 {
 	return { FMemberType(SizeType), DynamicStructType, 1, NoId, nullptr };
 }
 
-PLAINPROPS_API FMemberSchema MakeNestedRangeSchema(FScratchAllocator& Scratch, ERangeSizeType SizeType, FMemberSchema InnerRangeSchema);
+PLAINPROPS_API FMemberSchema MakeNestedRangeSchema(FScratchAllocator& Scratch, ESizeType SizeType, FMemberSchema InnerRangeSchema);
 
 //////////////////////////////////////////////////////////////////////////
 
 [[nodiscard]] PLAINPROPS_API FBuiltRange* CloneLeaves(FScratchAllocator& Scratch, uint64 Num, const void* Data, SIZE_T LeafSize);
 
 template<Arithmetic T>
-[[nodiscard]] FTypedRange BuildLeafRange(FScratchAllocator& Scratch, ERangeSizeType SizeType, const T* Values, uint64 Num)
+[[nodiscard]] FTypedRange BuildLeafRange(FScratchAllocator& Scratch, ESizeType SizeType, const T* Values, uint64 Num)
 {
 	// todo: detect invalid floats
 	return { MakeLeafRangeSchema<T>(SizeType), CloneLeaves(Scratch, Num, Values, sizeof(T)) };
@@ -204,7 +204,7 @@ template<Arithmetic T, typename SizeType>
 }
 
 template<Arithmetic T>
-[[nodiscard]] FTypedRange BuildLeafRange(FScratchAllocator& Scratch, ERangeSizeType SizeType, TConstArrayView<T> Values)
+[[nodiscard]] FTypedRange BuildLeafRange(FScratchAllocator& Scratch, ESizeType SizeType, TConstArrayView<T> Values)
 {
 	return BuildLeafRange(Scratch, SizeType, Values.GetData(), Values.Num());
 }
@@ -222,12 +222,12 @@ template<Enumeration T, typename SizeType>
 }
 
 template<UnsignedIntegral T>
-[[nodiscard]] FTypedRange BuildEnumRange(FScratchAllocator& Scratch, FEnumId Enum, ERangeSizeType SizeType, TConstArrayView<T> Values)
+[[nodiscard]] FTypedRange BuildEnumRange(FScratchAllocator& Scratch, FEnumId Enum, ESizeType SizeType, TConstArrayView<T> Values)
 {
 	return { MakeEnumRangeSchema<T>(Enum, SizeType), CloneLeaves(Scratch, Values.Num(), Values.GetData(), sizeof(T)) };
 }
 
-[[nodiscard]] inline FTypedRange MakeStructRange(FStructId Id, ERangeSizeType SizeType, FBuiltRange* Values )
+[[nodiscard]] inline FTypedRange MakeStructRange(FStructId Id, ESizeType SizeType, FBuiltRange* Values )
 {
 	return { MakeStructRangeSchema(SizeType, Id), Values };
 }
@@ -340,7 +340,7 @@ private:
 class FStructRangeBuilder
 {
 public:
-	FStructRangeBuilder(uint64 Num, ERangeSizeType InSizeType)
+	FStructRangeBuilder(uint64 Num, ESizeType InSizeType)
 	: SizeType(InSizeType)
 	{
 		Structs.SetNum(Num);
@@ -357,7 +357,7 @@ public:
 
 private:
 	TArray64<FMemberBuilder> Structs; 
-	ERangeSizeType SizeType;
+	ESizeType SizeType;
 };
 
 
@@ -379,7 +379,7 @@ public:
 		Ranges.Add(Range.Values);
 	}
 
-	[[nodiscard]] FTypedRange BuildAndReset(FScratchAllocator& Scratch, ERangeSizeType SizeType);
+	[[nodiscard]] FTypedRange BuildAndReset(FScratchAllocator& Scratch, ESizeType SizeType);
 
 private:
 	TArray64<FBuiltRange*> Ranges; 
diff --git a/Engine/Plugins/Experimental/PlainProps/Source/Public/PlainPropsIndex.h b/Engine/Plugins/Experimental/PlainProps/Source/Public/PlainPropsIndex.h
--- a/Engine/Plugins/Experimental/PlainProps/Source/Public/PlainPropsIndex.h
+++ b/Engine/Plugins/Experimental/PlainProps/Source/Public/PlainPropsIndex.h
@@ -57,7 +57,7 @@ public:
 	PLAINPROPS_API FType				MakeParametricType(FType Type, TConstArrayView<FType> Params);
 	PLAINPROPS_API FType				MakeAnonymousParametricType(TConstArrayView<FType> Params);
 	FType								MakeLeafParameter(FUnpackedLeafType Leaf)			{ return { NoId, FTypenameId(Leaves[uint8(Leaf.Type)][uint8(Leaf.Width)]) }; }
-	FType								MakeRangeParameter(ERangeSizeType SizeType)			{ return { NoId, FTypenameId(Ranges[uint8(SizeType)]) }; }
+	FType								MakeRangeParameter(ESizeType SizeType)				{ return { NoId, FTypenameId(Ranges[uint8(SizeType)]) }; }
 
 	PLAINPROPS_API FEnumId				IndexEnum(FType Type);
 	PLAINPROPS_API FStructId			IndexStruct(FType Type);
diff --git a/Engine/Plugins/Experimental/PlainProps/Source/Public/PlainPropsLoadMember.h b/Engine/Plugins/Experimental/PlainProps/Source/Public/PlainPropsLoadMember.h
--- a/Engine/Plugins/Experimental/PlainProps/Source/Public/PlainPropsLoadMember.h
+++ b/Engine/Plugins/Experimental/PlainProps/Source/Public/PlainPropsLoadMember.h
@@ -90,7 +90,7 @@ private:
 
 PLAINPROPS_API void	LoadRange(void* Dst, FRangeLoadView Src, TConstArrayView<FRangeBinding> InnerBindings);
 // Experimental low-level API bypassing FMemberLoader, internal use only
-PLAINPROPS_API void	LoadRange(void* Dst, FByteReader& SrcBytes, FBitCacheReader& SrcBits, ERangeSizeType MaxSize, FRangeLoadSchema Schema, TConstArrayView<FRangeBinding> InnerBindings);
+PLAINPROPS_API void	LoadRange(void* Dst, FByteReader& SrcBytes, FBitCacheReader& SrcBits, ESizeType MaxSize, FRangeLoadSchema Schema, TConstArrayView<FRangeBinding> InnerBindings);
 
 //////////////////////////////////////////////////////////////////////////
 
diff --git a/Engine/Plugins/Experimental/PlainProps/Source/Public/PlainPropsParse.h b/Engine/Plugins/Experimental/PlainProps/Source/Public/PlainPropsParse.h
--- a/Engine/Plugins/Experimental/PlainProps/Source/Public/PlainPropsParse.h
+++ b/Engine/Plugins/Experimental/PlainProps/Source/Public/PlainPropsParse.h
@@ -10,7 +10,7 @@ namespace PlainProps
 
 template<typename T>	PLAINPROPS_API TOptional<T>						Parse(FUtf8StringView String);
 
-template<>				PLAINPROPS_API TOptional<ERangeSizeType>		Parse(FUtf8StringView String);
+template<>				PLAINPROPS_API TOptional<ESizeType>				Parse(FUtf8StringView String);
 template<>				PLAINPROPS_API TOptional<ELeafWidth>			Parse(FUtf8StringView String);
 
 template<>				PLAINPROPS_API TOptional<FUnpackedLeafType>		Parse(FUtf8StringView String);
diff --git a/Engine/Plugins/Experimental/PlainProps/Source/Public/PlainPropsPrint.h b/Engine/Plugins/Experimental/PlainProps/Source/Public/PlainPropsPrint.h
--- a/Engine/Plugins/Experimental/PlainProps/Source/Public/PlainPropsPrint.h
+++ b/Engine/Plugins/Experimental/PlainProps/Source/Public/PlainPropsPrint.h
@@ -16,7 +16,7 @@ struct FReadDiffPath;
 struct FStructView;
 
 PLAINPROPS_API FAnsiStringView ToString(ELeafWidth Width);
-PLAINPROPS_API FAnsiStringView ToString(ERangeSizeType Width);
+PLAINPROPS_API FAnsiStringView ToString(ESizeType Width);
 PLAINPROPS_API FAnsiStringView ToString(ESchemaFormat Format);
 PLAINPROPS_API FAnsiStringView ToString(FUnpackedLeafType Leaf);
 
@@ -27,7 +27,7 @@ template<typename T> inline void Print(FUtf8Builder& Out, T Value)
 	static_assert(!sizeof(T), "Unsupported type for Print");
 }
 
-template<> PLAINPROPS_API void Print(FUtf8Builder& Out, ERangeSizeType Value);
+template<> PLAINPROPS_API void Print(FUtf8Builder& Out, ESizeType Value);
 template<> PLAINPROPS_API void Print(FUtf8Builder& Out, ELeafWidth Value);
 
 template<> PLAINPROPS_API void Print(FUtf8Builder& Out, bool Value);
diff --git a/Engine/Plugins/Experimental/PlainProps/Source/Public/PlainPropsSaveMember.h b/Engine/Plugins/Experimental/PlainProps/Source/Public/PlainPropsSaveMember.h
--- a/Engine/Plugins/Experimental/PlainProps/Source/Public/PlainPropsSaveMember.h
+++ b/Engine/Plugins/Experimental/PlainProps/Source/Public/PlainPropsSaveMember.h
@@ -81,7 +81,7 @@ struct TStructSerializer : FBothStructId
 		return SaveStruct(&Value, BindId, Ctx);
 	}
 
-	FMemberSchema MakeMemberRangeSchema(ERangeSizeType MaxSize) const
+	FMemberSchema MakeMemberRangeSchema(ESizeType MaxSize) const
 	{
 		return { FMemberType(MaxSize), DefaultStructType, 1, FInnerId(BindId), nullptr };
 	}
@@ -135,7 +135,7 @@ struct TRangeSerializer
 	{
 		for (uint16 Idx = 0; Idx < NumRanges; ++Idx)
 		{
-			ERangeSizeType MaxSize = Bindings[Idx].GetSizeType();
+			ESizeType MaxSize = Bindings[Idx].GetSizeType();
 			Types[Idx] = FMemberType(MaxSize);
 			BindTypes[Idx].Value = FMemberBindType(MaxSize);
 		}
@@ -231,7 +231,7 @@ struct TRangeSerializer
 
 	using BuiltItemType = FBuiltRange*;
 
-	FMemberSchema MakeMemberRangeSchema(ERangeSizeType MaxSize) const
+	FMemberSchema MakeMemberRangeSchema(ESizeType MaxSize) const
 	{
 		return { FMemberType(MaxSize), Types[0], NumTypes, InnermostId, Types };
 	}
@@ -307,7 +307,7 @@ struct TArithmeticSerializer : TLeafSerializer<T>
 
 	////////// TRangeSaver API //////////
 
-	FMemberSchema MakeMemberRangeSchema(ERangeSizeType MaxSize) const
+	FMemberSchema MakeMemberRangeSchema(ESizeType MaxSize) const
 	{
 		return { FMemberType(MaxSize), MemberType, 1, NoId, nullptr };
 	}
@@ -345,7 +345,7 @@ struct TEnumSerializer : TLeafSerializer<T>
 
 	////////// TRangeSaver API //////////
 
-	FMemberSchema MakeMemberRangeSchema(ERangeSizeType MaxSize) const
+	FMemberSchema MakeMemberRangeSchema(ESizeType MaxSize) const
 	{
 		return { FMemberType(MaxSize), MemberType, 1, FInnerId(Id), nullptr };
 	}
@@ -432,7 +432,7 @@ public:
 		FRangeSaverBase::AddBuiltItem(Schema.SaveItem(Item, Ctx));
 	}
 
-	[[nodiscard]] FTypedRange Finalize(ERangeSizeType MaxSize)
+	[[nodiscard]] FTypedRange Finalize(ESizeType MaxSize)
 	{
 		return FRangeSaverBase::Finalize(Schema.MakeMemberRangeSchema(MaxSize));
 	}
diff --git a/Engine/Plugins/Experimental/PlainProps/Source/Public/PlainPropsSpecify.h b/Engine/Plugins/Experimental/PlainProps/Source/Public/PlainPropsSpecify.h
--- a/Engine/Plugins/Experimental/PlainProps/Source/Public/PlainPropsSpecify.h
+++ b/Engine/Plugins/Experimental/PlainProps/Source/Public/PlainPropsSpecify.h
@@ -21,7 +21,7 @@ struct FMemberSpecView
 {
 	FMemberType						InnermostType;
 	FOptionalInnerId				InnermostId;
-	TConstArrayView<ERangeSizeType>	Ranges;
+	TConstArrayView<ESizeType>		Ranges;
 };
 
 // Specifies the type of a member before declaring it
@@ -45,17 +45,17 @@ public:
 	FMemberSpec(FDeclId Id)			: FMemberSpec(FStructType{EMemberKind::Struct, /* IsDynamic */ 0, /* IsSuper */ 0}, FInnerId(Id)) {}
 	FMemberSpec(ESpecDynamicStruct) : FMemberSpec(FStructType{EMemberKind::Struct, /* IsDynamic */ 1, /* IsSuper */ 0}, NoId) {}
 
-	PLAINPROPS_API FMemberSpec(ERangeSizeType MaxSize, FMemberSpec Inner);
+	PLAINPROPS_API FMemberSpec(ESizeType MaxSize, FMemberSpec Inner);
 	PLAINPROPS_API FMemberSpec(TConstArrayView<FMemberType> Members, FOptionalInnerId InInnermostId);
 	PLAINPROPS_API FMemberSpec(FMemberType Type, TConstArrayView<FMemberType> InnerRangeTypes, FOptionalInnerId InInnermostId);
 	
-	void RangeWrap(ERangeSizeType MaxSize)	{ Ranges.Emplace(MaxSize); }
+	void RangeWrap(ESizeType MaxSize)		{ Ranges.Emplace(MaxSize); }
 	void RangeWrap(FMemberType Range)		{ Ranges.Emplace(Range.AsRange().MaxSize); }
 
 	operator FMemberSpecView() const UE_LIFETIMEBOUND;
 	
 private:
-	using RangeSizeArray = TArray<ERangeSizeType, TInlineAllocator<8>>;
+	using RangeSizeArray = TArray<ESizeType, TInlineAllocator<8>>;
 	
 	bool							bSpecified = true;						
 	FMemberType						InnermostType;
diff --git a/Engine/Plugins/Experimental/PlainProps/Source/Public/PlainPropsTypes.h b/Engine/Plugins/Experimental/PlainProps/Source/Public/PlainPropsTypes.h
--- a/Engine/Plugins/Experimental/PlainProps/Source/Public/PlainPropsTypes.h
+++ b/Engine/Plugins/Experimental/PlainProps/Source/Public/PlainPropsTypes.h
@@ -15,7 +15,7 @@ namespace PlainProps
 enum class EMemberKind : uint8 { Leaf, Struct, Range };
 enum class ELeafType : uint8 { Bool, IntS, IntU, Float, Hex, Enum, Unicode };
 enum class ELeafWidth : uint8 { B8, B16, B32, B64 };
-enum class ERangeSizeType : uint8 { Uni, S8, U8, S16, U16, S32, U32, S64, U64 };
+enum class ESizeType : uint8 { Uni, S8, U8, S16, U16, S32, U32, S64, U64 };
 
 inline constexpr SIZE_T SizeOf(ELeafWidth Width)
 {
@@ -28,13 +28,13 @@ inline ELeafWidth WidthOf(SIZE_T Size)
 	return static_cast<ELeafWidth>(FMath::FloorLog2NonZero_64(Size));
 }
 
-inline SIZE_T SizeOf(ERangeSizeType Width)
+inline SIZE_T SizeOf(ESizeType Width)
 {
-	check(Width != ERangeSizeType::Uni);
+	check(Width != ESizeType::Uni);
     return SIZE_T(1) << ((uint8)Width - 1) / 2; 
 }
 
-inline constexpr uint64 Max(ERangeSizeType Width)
+inline constexpr uint64 Max(ESizeType Width)
 {
     const uint8_t LeadingZeroes[] = {63, 57, 56, 49, 48, 33, 32, 1, 0};
     return ~uint64_t(0) >> LeadingZeroes[(uint8)Width]; 
@@ -66,7 +66,7 @@ struct FLeafType
 struct FRangeType
 {
 	EMemberKind		_ : 2;
-	ERangeSizeType	MaxSize : 4;
+	ESizeType		MaxSize : 4;
 	uint8			_Pad : 2;
 };
 
@@ -85,7 +85,7 @@ union FMemberType
 	constexpr explicit FMemberType(FLeafType InLeaf) : Leaf(InLeaf) {}
 	constexpr explicit FMemberType(ELeafType Type, ELeafWidth Width) : Leaf({EMemberKind::Leaf, Width, Type}) {}
 	constexpr explicit FMemberType(FRangeType InRange) : Range(InRange) {}
-	constexpr explicit FMemberType(ERangeSizeType MaxSize) : Range({EMemberKind::Range, MaxSize}) {}
+	constexpr explicit FMemberType(ESizeType MaxSize) : Range({EMemberKind::Range, MaxSize}) {}
 	constexpr explicit FMemberType(FStructType InStruct) : Struct(InStruct) {}
 
 	bool			IsLeaf() const		{ return Kind == EMemberKind::Leaf; }
@@ -476,15 +476,15 @@ inline constexpr FUnpackedLeafType ReflectLeaf = MakeLeafType<T>();
 
 //////////////////////////////////////////////////////////////////////////
 
-inline constexpr ERangeSizeType RangeSizeOf(bool)	{ return ERangeSizeType::Uni; }
-inline constexpr ERangeSizeType RangeSizeOf(int8)	{ return ERangeSizeType::S8; }
-inline constexpr ERangeSizeType RangeSizeOf(int16)	{ return ERangeSizeType::S16; }
-inline constexpr ERangeSizeType RangeSizeOf(int32)	{ return ERangeSizeType::S32; }
-inline constexpr ERangeSizeType RangeSizeOf(int64)	{ return ERangeSizeType::S64; }
-inline constexpr ERangeSizeType RangeSizeOf(uint8)	{ return ERangeSizeType::U8; }
-inline constexpr ERangeSizeType RangeSizeOf(uint16)	{ return ERangeSizeType::U16; }
-inline constexpr ERangeSizeType RangeSizeOf(uint32)	{ return ERangeSizeType::U32; }
-inline constexpr ERangeSizeType RangeSizeOf(uint64)	{ return ERangeSizeType::U64; }
+inline constexpr ESizeType RangeSizeOf(bool)	{ return ESizeType::Uni; }
+inline constexpr ESizeType RangeSizeOf(int8)	{ return ESizeType::S8; }
+inline constexpr ESizeType RangeSizeOf(int16)	{ return ESizeType::S16; }
+inline constexpr ESizeType RangeSizeOf(int32)	{ return ESizeType::S32; }
+inline constexpr ESizeType RangeSizeOf(int64)	{ return ESizeType::S64; }
+inline constexpr ESizeType RangeSizeOf(uint8)	{ return ESizeType::U8; }
+inline constexpr ESizeType RangeSizeOf(uint16)	{ return ESizeType::U16; }
+inline constexpr ESizeType RangeSizeOf(uint32)	{ return ESizeType::U32; }
+inline constexpr ESizeType RangeSizeOf(uint64)	{ return ESizeType::U64; }
 
 //////////////////////////////////////////////////////////////////////////
 
diff --git a/Engine/Plugins/Experimental/PlainProps/Source/Public/PlainPropsUeCoreBindings.h b/Engine/Plugins/Experimental/PlainProps/Source/Public/PlainPropsUeCoreBindings.h
--- a/Engine/Plugins/Experimental/PlainProps/Source/Public/PlainPropsUeCoreBindings.h
+++ b/Engine/Plugins/Experimental/PlainProps/Source/Public/PlainPropsUeCoreBindings.h
@@ -733,7 +733,7 @@ struct TBaseDeltaBinding : ICustomBinding, FBaseDeltaBinding
 {
 	using Type = SetType;
 	using ElemType = typename SetType::ElementType;
-	static constexpr ERangeSizeType MaxSize = ERangeSizeType::S32;
+	static constexpr ESizeType MaxSize = ESizeType::S32;
 	static constexpr bool bIsSet = std::is_same_v<KeyType, ElemType>;
 	
 	const TMemberSerializer<ElemType> Elems;
diff --git a/Engine/Plugins/Experimental/PlainProps/Source/Tests/PlainPropsTests.cpp b/Engine/Plugins/Experimental/PlainProps/Source/Tests/PlainPropsTests.cpp
--- a/Engine/Plugins/Experimental/PlainProps/Source/Tests/PlainPropsTests.cpp
+++ b/Engine/Plugins/Experimental/PlainProps/Source/Tests/PlainPropsTests.cpp
@@ -491,15 +491,15 @@ TStaticArray<FMemberSpec, N> Respec(FMemberSpec Spec)
 	return MakeUniformStaticArray<FMemberSpec, N>(Spec);
 }
 
-static FMemberSpec UniRange(FMemberSpec Item)	{ return FMemberSpec(ERangeSizeType::Uni, Item); }
-static FMemberSpec S8Range(FMemberSpec Item)	{ return FMemberSpec(ERangeSizeType::S8, Item); }
-static FMemberSpec U8Range(FMemberSpec Item)	{ return FMemberSpec(ERangeSizeType::U8, Item); }
-static FMemberSpec S16Range(FMemberSpec Item)	{ return FMemberSpec(ERangeSizeType::S16, Item); }
-static FMemberSpec U16Range(FMemberSpec Item)	{ return FMemberSpec(ERangeSizeType::U16, Item); }
-static FMemberSpec S32Range(FMemberSpec Item)	{ return FMemberSpec(ERangeSizeType::S32, Item); }
-static FMemberSpec U32Range(FMemberSpec Item)	{ return FMemberSpec(ERangeSizeType::U32, Item); }
-static FMemberSpec S64Range(FMemberSpec Item)	{ return FMemberSpec(ERangeSizeType::S64, Item); }
-static FMemberSpec U64Range(FMemberSpec Item)	{ return FMemberSpec(ERangeSizeType::U64, Item); }
+static FMemberSpec UniRange(FMemberSpec Item)	{ return FMemberSpec(ESizeType::Uni, Item); }
+static FMemberSpec S8Range(FMemberSpec Item)	{ return FMemberSpec(ESizeType::S8, Item); }
+static FMemberSpec U8Range(FMemberSpec Item)	{ return FMemberSpec(ESizeType::U8, Item); }
+static FMemberSpec S16Range(FMemberSpec Item)	{ return FMemberSpec(ESizeType::S16, Item); }
+static FMemberSpec U16Range(FMemberSpec Item)	{ return FMemberSpec(ESizeType::U16, Item); }
+static FMemberSpec S32Range(FMemberSpec Item)	{ return FMemberSpec(ESizeType::S32, Item); }
+static FMemberSpec U32Range(FMemberSpec Item)	{ return FMemberSpec(ESizeType::U32, Item); }
+static FMemberSpec S64Range(FMemberSpec Item)	{ return FMemberSpec(ESizeType::S64, Item); }
+static FMemberSpec U64Range(FMemberSpec Item)	{ return FMemberSpec(ESizeType::U64, Item); }
 
 inline FMemberSpec Enum8(FEnumId Id)			{ return FMemberSpec({ELeafType::Enum, ELeafWidth::B8}, Id); }
 inline FMemberSpec Enum16(FEnumId Id)			{ return FMemberSpec({ELeafType::Enum, ELeafWidth::B16}, Id); }
@@ -1211,28 +1211,28 @@ TEST_CASE_NAMED(FPlainPropsReadWriteTest, "System::Core::Serialization::PlainPro
 			XYs[0].Add(Batch.NameMember("Y"), 2.f);
 			XYs[1].Add(Batch.NameMember("X"), 3.f);
 			XYs[1].Add(Batch.NameMember("Y"), 4.f);
-			FNestedRangeBuilder StructRs(MakeStructRangeSchema(ERangeSizeType::U64, XY), 1);
+			FNestedRangeBuilder StructRs(MakeStructRangeSchema(ESizeType::U64, XY), 1);
 			StructRs.Add(XYs.BuildAndReset(Scratch, Batch.Get(XY), Batch.GetDebug()));
 
 			FStructRangeBuilder ZWs(int16(3));
 			ZWs[0].Add(Batch.NameMember("Z"), 1.5f);
 			ZWs[2].Add(Batch.NameMember("Z"), 2.5f);
 			ZWs[2].Add(Batch.NameMember("W"), 3.5f);
-			FMemberSchema ZWRangeSchema = MakeStructRangeSchema(ERangeSizeType::S16, ZW);
+			FMemberSchema ZWRangeSchema = MakeStructRangeSchema(ESizeType::S16, ZW);
 			FNestedRangeBuilder ZWRs(ZWRangeSchema, 1);
 			ZWRs.Add(ZWs.BuildAndReset(Scratch, Batch.Get(ZW), Batch.GetDebug()));
-			FNestedRangeBuilder StructRRs(MakeNestedRangeSchema(Scratch, ERangeSizeType::U32, ZWRangeSchema), 1);
-			StructRRs.Add(ZWRs.BuildAndReset(Scratch, ERangeSizeType::U32));
+			FNestedRangeBuilder StructRRs(MakeNestedRangeSchema(Scratch, ESizeType::U32, ZWRangeSchema), 1);
+			StructRRs.Add(ZWRs.BuildAndReset(Scratch, ESizeType::U32));
 
 			FMemberBuilder Members;
-			Members.AddRange(Batch.NameMember("IntRs"), IntRs.BuildAndReset(Scratch, ERangeSizeType::S32));
-			Members.AddRange(Batch.NameMember("EmptyRs"), IntRs.BuildAndReset(Scratch, ERangeSizeType::S32));
-			Members.AddRange(Batch.NameMember("FloatRs"), FloatRs.BuildAndReset(Scratch, ERangeSizeType::S64));
-			Members.AddRange(Batch.NameMember("EnumRs"), EnumRs.BuildAndReset(Scratch, ERangeSizeType::U8));
-			Members.AddRange(Batch.NameMember("UnusedEnumRs"), UnusedEnumRs.BuildAndReset(Scratch, ERangeSizeType::U8));
-			Members.AddRange(Batch.NameMember("UnicodeRs"), UnicodeRs.BuildAndReset(Scratch, ERangeSizeType::S32));
-			Members.AddRange(Batch.NameMember("StructRs"), StructRs.BuildAndReset(Scratch, ERangeSizeType::U64));
-			Members.AddRange(Batch.NameMember("StructRRs"), StructRRs.BuildAndReset(Scratch, ERangeSizeType::U32));
+			Members.AddRange(Batch.NameMember("IntRs"), IntRs.BuildAndReset(Scratch, ESizeType::S32));
+			Members.AddRange(Batch.NameMember("EmptyRs"), IntRs.BuildAndReset(Scratch, ESizeType::S32));
+			Members.AddRange(Batch.NameMember("FloatRs"), FloatRs.BuildAndReset(Scratch, ESizeType::S64));
+			Members.AddRange(Batch.NameMember("EnumRs"), EnumRs.BuildAndReset(Scratch, ESizeType::U8));
+			Members.AddRange(Batch.NameMember("UnusedEnumRs"), UnusedEnumRs.BuildAndReset(Scratch, ESizeType::U8));
+			Members.AddRange(Batch.NameMember("UnicodeRs"), UnicodeRs.BuildAndReset(Scratch, ESizeType::S32));
+			Members.AddRange(Batch.NameMember("StructRs"), StructRs.BuildAndReset(Scratch, ESizeType::U64));
+			Members.AddRange(Batch.NameMember("StructRRs"), StructRRs.BuildAndReset(Scratch, ESizeType::U32));
 
 			Batch.AddObject(Object, MoveTemp(Members));
 		}, 
@@ -1314,10 +1314,10 @@ TEST_CASE_NAMED(FPlainPropsReadWriteTest, "System::Core::Serialization::PlainPro
 			Structs[7].AddRange(Batch.NameMember("MaybeB"), BuildLeafRange(Scratch, &True, true));
 			Structs[7].AddRange(Batch.NameMember("Bs"),		BuildLeafRange(Scratch, MakeArrayView({true, true, false, false, true, true, false, false, true, true})));
 			MaybeBs.Add(BuildLeafRange(Scratch, &True, true));
-			Structs[7].AddRange(Batch.NameMember("MaybeBs"), MaybeBs.BuildAndReset(Scratch, ERangeSizeType::Uni));
+			Structs[7].AddRange(Batch.NameMember("MaybeBs"), MaybeBs.BuildAndReset(Scratch, ESizeType::Uni));
 			Structs[7].Add(Batch.NameMember("B"), true);
 			MaybeBs.Add(BuildLeafRange(Scratch, &True, false));
-			Structs[8].AddRange(Batch.NameMember("MaybeBs"), MaybeBs.BuildAndReset(Scratch, ERangeSizeType::Uni));
+			Structs[8].AddRange(Batch.NameMember("MaybeBs"), MaybeBs.BuildAndReset(Scratch, ESizeType::Uni));
 			Structs[9].Add(Batch.NameMember("B"), false);
 
 			FMemberBuilder Members;
@@ -1429,9 +1429,9 @@ TEST_CASE_NAMED(FPlainPropsReadWriteTest, "System::Core::Serialization::PlainPro
 			O1.AddRange(Batch.NameMember("Diff"), BuildStructRange(SA, 12));
 			O1.AddRange(Batch.NameMember("SameEmpty"), BuildStructRange(SA, 13));
 			O1.AddRange(Batch.NameMember("DiffEmpty"), BuildStructRange(SA, 14));
-			FNestedRangeBuilder NestedSA(MakeStructRangeSchema(ERangeSizeType::S32, SA), 1);
+			FNestedRangeBuilder NestedSA(MakeStructRangeSchema(ESizeType::S32, SA), 1);
 			NestedSA.Add(BuildStructRange(SA, 100));
-			O1.AddRange(Batch.NameMember("DiffNested"), NestedSA.BuildAndReset(Scratch, ERangeSizeType::S32));
+			O1.AddRange(Batch.NameMember("DiffNested"), NestedSA.BuildAndReset(Scratch, ESizeType::S32));
 			
 			FMemberBuilder O2;
 			O2.AddRange(Batch.NameMember("Same"), BuildStructRange(SA, 20));
@@ -1439,9 +1439,9 @@ TEST_CASE_NAMED(FPlainPropsReadWriteTest, "System::Core::Serialization::PlainPro
 			O2.AddRange(Batch.NameMember("SameEmpty"), FStructRangeBuilder(0).BuildAndReset(Scratch, Batch.Get(SA), Batch.GetDebug()));
 			// PP-TEXT: Handle DiffEmpty in DynamicStructRange Test by printing dynamic type for empty ranges
 			// O2.AddRange(Batch.NameMember("DiffEmpty"), FStructRangeBuilder(0).BuildAndReset(Scratch, Batch.Get(SB), Batch.GetDebug()));
-			FNestedRangeBuilder NestedSB(MakeStructRangeSchema(ERangeSizeType::S32, SB), 1);
+			FNestedRangeBuilder NestedSB(MakeStructRangeSchema(ESizeType::S32, SB), 1);
 			NestedSB.Add(BuildStructRange(SB, 200.f));
-			O2.AddRange(Batch.NameMember("DiffNested"), NestedSB.BuildAndReset(Scratch, ERangeSizeType::S32));
+			O2.AddRange(Batch.NameMember("DiffNested"), NestedSB.BuildAndReset(Scratch, ESizeType::S32));
 			
 			Batch.AddObject(Object, MoveTemp(O1));
 			Batch.AddObject(Object, MoveTemp(O2));
diff --git a/Engine/Plugins/Experimental/PlainPropsUObject/Source/Private/PlainPropsUObjectBindings.cpp b/Engine/Plugins/Experimental/PlainPropsUObject/Source/Private/PlainPropsUObjectBindings.cpp
--- a/Engine/Plugins/Experimental/PlainPropsUObject/Source/Private/PlainPropsUObjectBindings.cpp
+++ b/Engine/Plugins/Experimental/PlainPropsUObject/Source/Private/PlainPropsUObjectBindings.cpp
@@ -87,7 +87,7 @@ DEFINE_LOG_CATEGORY_STATIC(LogPlainPropsUObject, Log, All);
 namespace PlainProps::UE
 {
 
-static constexpr ERangeSizeType DefaultRangeMax = RangeSizeOf(FDefaultAllocator::SizeType{});
+static constexpr ESizeType DefaultRangeMax = RangeSizeOf(FDefaultAllocator::SizeType{});
 
 inline FMemberSpec DefaultRangeOf(FMemberSpec Item) { return FMemberSpec(DefaultRangeMax, Item); }
 
@@ -1181,7 +1181,7 @@ struct TLeafArrayBinding : ILeafRangeBinding
 // Reusable cache of FArrayProperty range bindings
 class FArrayPropertyBindings
 {
-	static constexpr ERangeSizeType SizeType = DefaultRangeMax;
+	static constexpr ESizeType SizeType = DefaultRangeMax;
 	
 	const TLeafArrayBinding<ELeafType::Bool, ELeafWidth::B8>	Bool;
 	const TLeafArrayBinding<ELeafType::Float, ELeafWidth::B32>	Float;
@@ -1270,7 +1270,7 @@ static FRangeBinding AllocateArrayBinding(FArrayProperty* Property)
 	// Todo: Ownership / memory leak, try make non-trivial case cacheable by making FProperty ctor/dtor extractable
 	IItemRangeBinding* Out = Info.bFreezable	? CreateAndLeak<FFreezableScriptArray>(Info, Inner)
 												: CreateAndLeak<FScriptArray>(Info, Inner);
-	return FRangeBinding(*Out, ERangeSizeType::S32);
+	return FRangeBinding(*Out, ESizeType::S32);
 }
 
 ////////////////////////////////////////////////////////////////////////////////////////////////
@@ -2875,7 +2875,7 @@ public:
 
 		// Todo: Ownership / memory leak
 		const IItemRangeBinding* Out = bIntrusive ? new FIntrusiveOptionalBinding(Inner) : BindNew(Inner);
-		return Bindings.Emplace(Key, FRangeBinding(*Out, ERangeSizeType::Uni));
+		return Bindings.Emplace(Key, FRangeBinding(*Out, ESizeType::Uni));
 	}
 };
 
@@ -3261,7 +3261,7 @@ private:
 
 			// Todo: Ownership / memory leak
 			const FStaticArrayBinding& ItemBinding = *new FStaticArrayBinding(ArrayDim, ElementSize);
-			ERangeSizeType SizeType = ArrayDim < 256 ? ERangeSizeType::U8 : ((ArrayDim < 65536) ? ERangeSizeType::U16 : ERangeSizeType::U32);
+			ESizeType SizeType = ArrayDim < 256 ? ESizeType::U8 : ((ArrayDim < 65536) ? ESizeType::U16 : ESizeType::U32);
 			Out.RangeBindings =	AllocateRangeBindings(FRangeBinding(ItemBinding, SizeType), Out.RangeBindings);
 		}
 		else
