PlainProps: Namespace compile fix for ::UE::PreciseFPEqual()

#rnx

diff --git a/Engine/Plugins/Experimental/PlainProps/Source/Private/PlainPropsInternalDiff.h b/Engine/Plugins/Experimental/PlainProps/Source/Private/PlainPropsInternalDiff.h
--- a/Engine/Plugins/Experimental/PlainProps/Source/Private/PlainPropsInternalDiff.h
+++ b/Engine/Plugins/Experimental/PlainProps/Source/Private/PlainPropsInternalDiff.h
@@ -50,8 +50,8 @@ inline int64 DiffLeaf(const uint8* A, const uint8* B, FUnpackedLeafBindType Leaf
 		
 		switch (Leaf.Width)
 		{
-		case ELeafWidth::B32:	return !UE::PreciseFPEqual(reinterpret_cast<const float&>(*A), reinterpret_cast<const float&>(*B));
-		case ELeafWidth::B64:	return !UE::PreciseFPEqual(reinterpret_cast<const double&>(*A), reinterpret_cast<const double&>(*B));
+		case ELeafWidth::B32:	return !::UE::PreciseFPEqual(reinterpret_cast<const float&>(*A), reinterpret_cast<const float&>(*B));
+		case ELeafWidth::B64:	return !::UE::PreciseFPEqual(reinterpret_cast<const double&>(*A), reinterpret_cast<const double&>(*B));
 		default:				break;
 		}
 
