PlainProps: Remove and replace the Core scope with CoreUObject to reflect the UE object paths like /Script/CoreUObject.Transform.
Enable custom bindings for FTransform (that use FVector and FQuat), FGuid (that use SpecHex32) FColor and FLinearColor in the correct CoreUObject scope.
Add missing BuildHexRange to the Build api.
Add missing Print, Parse and Diff implementations for Hex.
Add unit tests for Hex leaves and ranges.

#rnx

diff --git a/Engine/Plugins/Experimental/PlainProps/Source/Private/PlainPropsDiff.cpp b/Engine/Plugins/Experimental/PlainProps/Source/Private/PlainPropsDiff.cpp
--- a/Engine/Plugins/Experimental/PlainProps/Source/Private/PlainPropsDiff.cpp
+++ b/Engine/Plugins/Experimental/PlainProps/Source/Private/PlainPropsDiff.cpp
@@ -389,13 +389,17 @@ static bool DiffLeaves(FUnpackedLeafType Leaf, FLeafRangeView A, FLeafRangeView
 		check(Leaf.Width == ELeafWidth::B64);
 		return DiffItems(A.AsDoubles(), B.AsDoubles(), OutIdx);
 	case ELeafType::Hex:
-		// PP-TEXT: Implement DiffItems(Hex)
-		check(Leaf.Type != ELeafType::Hex);
-		break;
+		switch (Leaf.Width)
+		{
+		case ELeafWidth::B8:  return DiffItems(A.AsHex8s(),  B.AsHex8s(),  OutIdx);
+		case ELeafWidth::B16: return DiffItems(A.AsHex16s(), B.AsHex16s(), OutIdx);
+		case ELeafWidth::B32: return DiffItems(A.AsHex32s(), B.AsHex32s(), OutIdx);
+		case ELeafWidth::B64: return DiffItems(A.AsHex64s(), B.AsHex64s(), OutIdx);
+		}
 	case ELeafType::Enum:
 		switch (Leaf.Width)
 		{
-		case ELeafWidth::B8:  return DiffItems(A.AsUnderlyingValues<uint8>(), B.AsUnderlyingValues<uint8>(), OutIdx);
+		case ELeafWidth::B8:  return DiffItems(A.AsUnderlyingValues<uint8>(),  B.AsUnderlyingValues<uint8>(),  OutIdx);
 		case ELeafWidth::B16: return DiffItems(A.AsUnderlyingValues<uint16>(), B.AsUnderlyingValues<uint16>(), OutIdx);
 		case ELeafWidth::B32: return DiffItems(A.AsUnderlyingValues<uint32>(), B.AsUnderlyingValues<uint32>(), OutIdx);
 		case ELeafWidth::B64: return DiffItems(A.AsUnderlyingValues<uint64>(), B.AsUnderlyingValues<uint64>(), OutIdx);
diff --git a/Engine/Plugins/Experimental/PlainProps/Source/Private/PlainPropsParse.cpp b/Engine/Plugins/Experimental/PlainProps/Source/Private/PlainPropsParse.cpp
--- a/Engine/Plugins/Experimental/PlainProps/Source/Private/PlainPropsParse.cpp
+++ b/Engine/Plugins/Experimental/PlainProps/Source/Private/PlainPropsParse.cpp
@@ -137,20 +137,30 @@ TOptional<ELeafWidth> Parse(FUtf8StringView String)
 template<Arithmetic T> 
 TOptional<T> Parse(FUtf8StringView String) requires (std::is_integral_v<T>)
 {
-	const char* first = reinterpret_cast<const char*>(String.GetData());
-	const char* last = first + String.Len();
+	const char* First = reinterpret_cast<const char*>(String.GetData());
+	const char* Last = First + String.Len();
 	T Out;
-	std::from_chars_result Result = std::from_chars(first, last, Out, 10);
+	std::from_chars_result Result = std::from_chars(First, Last, Out, 10);
+	return Result.ec == std::errc{} ? Out : TOptional<T>{ NullOpt };
+}
+
+template<UnsignedIntegral T> 
+TOptional<T> ParseHex(FUtf8StringView String)
+{
+	const char* First = reinterpret_cast<const char*>(String.GetData());
+	const char* Last = First + String.Len();
+	T Out;
+	std::from_chars_result Result = std::from_chars(First, Last, Out, 16);
 	return Result.ec == std::errc{} ? Out : TOptional<T>{ NullOpt };
 }
 
 template<Arithmetic T> 
 TOptional<T> Parse(FUtf8StringView String) requires (std::is_floating_point_v<T>)
 {
-	const char* first = reinterpret_cast<const char*>(String.GetData());
-	const char* last = first + String.Len();
+	const char* First = reinterpret_cast<const char*>(String.GetData());
+	const char* Last = First + String.Len();
 	T Out;
-	std::from_chars_result Result = std::from_chars(first, last, Out, std::chars_format::general);
+	std::from_chars_result Result = std::from_chars(First, Last, Out, std::chars_format::general);
 	return Result.ec == std::errc{} ? Out : TOptional<T>{ NullOpt };
 }
 
@@ -1373,7 +1383,7 @@ bool FMemberParser::ParseAll()
 	return true;
 }
 
-template<typename T>
+template<Arithmetic T>
 void AddMember(FMemberBuilder& Members, FMemberId Id, FUtf8StringView String)
 {
 	if (TOptional<T> Value = Parse<T>(String))
@@ -1382,7 +1392,7 @@ void AddMember(FMemberBuilder& Members, FMemberId Id, FUtf8StringView String)
 	}
 }
 
-template<typename T>
+template<UnsignedIntegral T>
 void AddEnum(FMemberBuilder& Members, FMemberId Id, FEnumId Enum, FUtf8StringView String)
 {
 	if (TOptional<T> Value = Parse<T>(String))
@@ -1391,6 +1401,15 @@ void AddEnum(FMemberBuilder& Members, FMemberId Id, FEnumId Enum, FUtf8StringVie
 	}
 }
 
+template<UnsignedIntegral T>
+void AddHex(FMemberBuilder& Members, FMemberId Id, FUtf8StringView String)
+{
+	if (TOptional<T> Value = ParseHex<T>(String))
+	{
+		Members.AddHex(Id, Value.GetValue());
+	}
+}
+
 static bool IsUnicodeString(const FParsedMemberSchema& MemberSchema)
 {
 	return MemberSchema.Type.IsRange() &&
@@ -1456,7 +1475,13 @@ bool FMemberParser::ParseLeaf()
 			}
 			break;
 		case ELeafType::Hex:
-			check(Leaf.Type != ELeafType::Hex);
+			switch (Leaf.Width)
+			{
+				case ELeafWidth::B8:	AddHex<uint8 >(MemberBuilder, MemberId, String);	break;
+				case ELeafWidth::B16:	AddHex<uint16>(MemberBuilder, MemberId, String);	break;
+				case ELeafWidth::B32:	AddHex<uint32>(MemberBuilder, MemberId, String);	break;
+				case ELeafWidth::B64:	AddHex<uint64>(MemberBuilder, MemberId, String);	break;
+			}
 			break;
 		case ELeafType::Enum:
 			switch (Leaf.Width)
@@ -1619,7 +1644,7 @@ template<Arithmetic T>
 	return BuildLeafRange(Scratch, MaxSize, TConstArrayView<T>(Values));
 }
 
-template<Arithmetic T>
+template<UnsignedIntegral T>
 [[nodiscard]] static FTypedRange ParseEnumRange(
 	FScratchAllocator& Scratch,
 	FEnumId Enum,
@@ -1638,6 +1663,24 @@ template<Arithmetic T>
 	return BuildEnumRange(Scratch, Enum, MaxSize, TConstArrayView<T>(Values));
 }
 
+template<UnsignedIntegral T>
+[[nodiscard]] static FTypedRange ParseHexRange(
+	FScratchAllocator& Scratch,
+	ESizeType MaxSize,
+	FScopedTokenReader& RangeTokens)
+{
+	TArray<T, TInlineAllocator64<64>> Values;
+	while (RangeTokens.HasMore())
+	{
+		FToken Token = RangeTokens.GrabToken();
+		if (TOptional<T> Value = ParseHex<T>(Token.Value()))
+		{
+			Values.Add(Value.GetValue());
+		}
+	}
+	return BuildHexRange(Scratch, MaxSize, TConstArrayView<T>(Values));
+}
+
 FTypedRange FMemberParser::ParseLeaves(const FMemberSchema& RangeSchema, FScopedTokenReader RangeTokens)
 {
 	const FOptionalEnumId Enum = ToOptionalEnum(RangeSchema.InnerSchema);
@@ -1681,8 +1724,13 @@ FTypedRange FMemberParser::ParseLeaves(const FMemberSchema& RangeSchema, FScoped
 			}
 			break;
 		case ELeafType::Hex:
-			// PP-TEXT: Implement ParseHexRange()
-			check(Leaf.Type != ELeafType::Hex);
+			switch (Leaf.Width)
+			{
+				case ELeafWidth::B8:	TypedRange = ParseHexRange<uint8 >(Scratch, NumType, RangeTokens); break;
+				case ELeafWidth::B16:	TypedRange = ParseHexRange<uint16>(Scratch, NumType, RangeTokens); break;
+				case ELeafWidth::B32:	TypedRange = ParseHexRange<uint32>(Scratch, NumType, RangeTokens); break;
+				case ELeafWidth::B64:	TypedRange = ParseHexRange<uint64>(Scratch, NumType, RangeTokens); break;
+			}
 			break;
 		case ELeafType::Enum:
 			switch (Leaf.Width)
@@ -1692,9 +1740,6 @@ FTypedRange FMemberParser::ParseLeaves(const FMemberSchema& RangeSchema, FScoped
 				case ELeafWidth::B32:	TypedRange = ParseEnumRange<uint32>(Scratch, Enum.Get(), NumType, RangeTokens); break;
 				case ELeafWidth::B64:	TypedRange = ParseEnumRange<uint64>(Scratch, Enum.Get(), NumType, RangeTokens); break;
 			}
-			{
-				FTypedRange Range = BuildEnumRange(Scratch, Enum.Get(), NumType, TConstArrayView<unsigned char>({'a', 'b'}));
-			}
 			break;
 		case ELeafType::Unicode:
 			checkf(!RangeTokens.HasMore(), TEXT("Should have been handled by PrintUnicodeRangeAsLeafValue/ParseUnicodeLeafValueAsRange"));
@@ -2089,4 +2134,4 @@ void ParseYamlBatch(TArray64<uint8>& OutBinary, FUtf8StringView Yaml)
 	BatchParser.Parse(OutBinary);
 }
 
-} // namespace PlainProps
\ No newline at end of file
+} // namespace PlainProps
diff --git a/Engine/Plugins/Experimental/PlainProps/Source/Private/PlainPropsPrint.cpp b/Engine/Plugins/Experimental/PlainProps/Source/Private/PlainPropsPrint.cpp
--- a/Engine/Plugins/Experimental/PlainProps/Source/Private/PlainPropsPrint.cpp
+++ b/Engine/Plugins/Experimental/PlainProps/Source/Private/PlainPropsPrint.cpp
@@ -99,21 +99,31 @@ void Print(FUtf8Builder& Out, T Value)
 template<Arithmetic T> requires (std::is_integral_v<T>)
 void Print(FUtf8Builder& Out, T Value)
 {
-	constexpr size_t BufSize = std::numeric_limits<T>::digits10 + 3;
-	char Buf[BufSize]{};
-	std::to_chars_result result = std::to_chars(Buf, Buf + BufSize, Value);
-	check(result.ec == std::errc());
-	Out.Append(Buf);
+	constexpr size_t Size = std::numeric_limits<T>::digits10 + 3;
+	char Buf[Size];
+	std::to_chars_result Result = std::to_chars(Buf, Buf + Size, Value);
+	check(Result.ec == std::errc());
+	Out.Append(Buf, Result.ptr - Buf);
+}
+
+template<UnsignedIntegral T>
+void PrintHex(FUtf8Builder& Out, T Value)
+{
+	constexpr size_t Size = sizeof(T) * 2;
+	char Buf[Size];
+	std::to_chars_result Result = std::to_chars(Buf, Buf + Size, Value, 16);
+	check(Result.ec == std::errc());
+	Out.Append(Buf, Result.ptr - Buf);
 }
 
 template<Arithmetic T> requires (std::is_floating_point_v<T>)
 void Print(FUtf8Builder& Out, T Value)
 {
-	constexpr size_t BufSize = 32;
-	char Buf[BufSize]{};
-	std::to_chars_result result = std::to_chars(Buf, Buf + BufSize, Value, std::chars_format::general);
-	check(result.ec == std::errc());
-	Out.Append(Buf);
+	constexpr size_t Size = 32;
+	char Buf[Size];
+	std::to_chars_result Result = std::to_chars(Buf, Buf + Size, Value, std::chars_format::general);
+	check(Result.ec == std::errc());
+	Out.Append(Buf, Result.ptr - Buf);
 }
 
 template<>
@@ -342,13 +352,22 @@ public:
 	void BeginRange();
 	void EndRange();
 
-	void AddLeafId(FUtf8StringView Id);
+	// split calls for id and value pairs
+	void AddId(FUtf8StringView Id);
 	template<typename T>
-	void AddLeafValue(T Value);
+	void AddValue(T Value);
+	template<UnsignedIntegral T>
+	void AddValueHex(T Value);
+
+	// id value pairs
 	template<typename T>
-	void AddLeaf(FUtf8StringView Id, T Value) { AddLeafId(Id); AddLeafValue(Value); }
+	void AddLeaf(FUtf8StringView Id, T Value) { AddId(Id); AddValue(Value); }
+
+	// range items
 	template<typename T>
-	void AddLeaf(T Value);
+	void AddItem(T Value);
+	template<UnsignedIntegral T>
+	void AddItemHex(T Value);
 
 	void AddComment(FUtf8StringView Comment);
 
@@ -360,6 +379,8 @@ private:
 	void AppendValue(T Value);
 	template<>
 	void AppendValue(FUtf8StringView Value);
+	template<UnsignedIntegral T>
+	void AppendValueHex(T Value);
 
 	struct FStackInfo
 	{
@@ -648,7 +669,7 @@ void FYamlBuilder::EndRange()
 	--IndentLevel;
 }
 
-void FYamlBuilder::AddLeafId(FUtf8StringView Id)
+void FYamlBuilder::AddId(FUtf8StringView Id)
 {
 	AppendNewLine();
 	AppendIndentation();
@@ -658,7 +679,7 @@ void FYamlBuilder::AddLeafId(FUtf8StringView Id)
 }
 
 template<typename T>
-void FYamlBuilder::AddLeafValue(T Value)
+void FYamlBuilder::AddValue(T Value)
 {
 	Text.AppendChar(' ');
 	AppendValue(Value);
@@ -666,8 +687,17 @@ void FYamlBuilder::AddLeafValue(T Value)
 	Stack.Last().IsEmpty = false;
 }
 
+template<UnsignedIntegral T>
+void FYamlBuilder::AddValueHex(T Value)
+{
+	Text.AppendChar(' ');
+	AppendValueHex(Value);
+	IsNewLine = false;
+	Stack.Last().IsEmpty = false;
+}
+
 template<typename T>
-void FYamlBuilder::AddLeaf(T Value)
+void FYamlBuilder::AddItem(T Value)
 {
 	AppendNewLine();
 	AppendIndentation();
@@ -676,6 +706,16 @@ void FYamlBuilder::AddLeaf(T Value)
 	Stack.Last().IsEmpty = false;
 }
 
+template<UnsignedIntegral T>
+void FYamlBuilder::AddItemHex(T Value)
+{
+	AppendNewLine();
+	AppendIndentation();
+	AppendValueHex(Value);
+	IsNewLine = false;
+	Stack.Last().IsEmpty = false;
+}
+
 void FYamlBuilder::AddComment(FUtf8StringView Comment)
 {
 	Text << " #" << Comment;
@@ -750,6 +790,14 @@ void FYamlBuilder::AppendValue(FUtf8StringView Value)
 	PrintQuotedString(Text, Value);
 }
 
+template<UnsignedIntegral T>
+void FYamlBuilder::AppendValueHex(T Value)
+{
+	Text.AppendChar('\'');
+	PrintHex(Text, Value);
+	Text.AppendChar('\'');
+}
+
 ///////////////////////////////////////////////////////////////////////////////
 
 FBatchPrinter::FBatchPrinter(FYamlBuilder& InTextBuilder, const FBatchIds& InIds)
@@ -829,11 +877,11 @@ void FBatchPrinter::PrintStructSchema(const FStructSchema& Struct, FSchemaBatchI
 	while (Reader.HasMore())
 	{
 		Ids.AppendString(Buf, Reader.PeekName());
-		TextBuilder.AddLeafId(Buf);
+		TextBuilder.AddId(Buf);
 		Buf.Reset();
 
 		PrintMemberSchema(Buf, Ids, Reader.GrabMember());
-		TextBuilder.AddLeafValue(FUtf8StringView(Buf));
+		TextBuilder.AddValue(FUtf8StringView(Buf));
 		Buf.Reset();
 	}
 	TextBuilder.EndRange();
@@ -913,60 +961,66 @@ void FMemberPrinter::PrintMembers(FStructView StructView)
 
 void FMemberPrinter::PrintLeaf(FMemberId Id, FLeafView LeafView)
 {
-	TextBuilder.AddLeafId(*FPrintId<128>(Ids, Id));
+	TextBuilder.AddId(*FPrintId<128>(Ids, Id));
 
 	switch (LeafView.Leaf.Type)
 	{
 	case ELeafType::Bool:
-		TextBuilder.AddLeafValue(LeafView.AsBool());
+		TextBuilder.AddValue(LeafView.AsBool());
 		break;
 	case ELeafType::IntS:
 		switch (LeafView.Leaf.Width)
 		{
-			case ELeafWidth::B8:	TextBuilder.AddLeafValue(LeafView.AsS8()); break;
-			case ELeafWidth::B16:	TextBuilder.AddLeafValue(LeafView.AsS16()); break;
-			case ELeafWidth::B32:	TextBuilder.AddLeafValue(LeafView.AsS32()); break;
-			case ELeafWidth::B64:	TextBuilder.AddLeafValue(LeafView.AsS64()); break;
+			case ELeafWidth::B8:	TextBuilder.AddValue(LeafView.AsS8()); break;
+			case ELeafWidth::B16:	TextBuilder.AddValue(LeafView.AsS16()); break;
+			case ELeafWidth::B32:	TextBuilder.AddValue(LeafView.AsS32()); break;
+			case ELeafWidth::B64:	TextBuilder.AddValue(LeafView.AsS64()); break;
 		}
 		break;
 	case ELeafType::IntU:
 		switch (LeafView.Leaf.Width)
 		{
-			case ELeafWidth::B8:	TextBuilder.AddLeafValue(LeafView.AsU8()); break;
-			case ELeafWidth::B16:	TextBuilder.AddLeafValue(LeafView.AsU16()); break;
-			case ELeafWidth::B32:	TextBuilder.AddLeafValue(LeafView.AsU32()); break;
-			case ELeafWidth::B64:	TextBuilder.AddLeafValue(LeafView.AsU64()); break;
+			case ELeafWidth::B8:	TextBuilder.AddValue(LeafView.AsU8()); break;
+			case ELeafWidth::B16:	TextBuilder.AddValue(LeafView.AsU16()); break;
+			case ELeafWidth::B32:	TextBuilder.AddValue(LeafView.AsU32()); break;
+			case ELeafWidth::B64:	TextBuilder.AddValue(LeafView.AsU64()); break;
 		}
 		break;
 	case ELeafType::Float:
 		if (LeafView.Leaf.Width == ELeafWidth::B32)
 		{
-			TextBuilder.AddLeafValue(LeafView.AsFloat());
+			TextBuilder.AddValue(LeafView.AsFloat());
 		}
 		else
 		{
 			check(LeafView.Leaf.Width == ELeafWidth::B64);
-			TextBuilder.AddLeafValue(LeafView.AsDouble());
+			TextBuilder.AddValue(LeafView.AsDouble());
 		}
 		break;
 	case ELeafType::Hex:
-		check(LeafView.Leaf.Type != ELeafType::Hex);
+		switch (LeafView.Leaf.Width)
+		{
+			case ELeafWidth::B8:	TextBuilder.AddValueHex(LeafView.AsHex8());  break;
+			case ELeafWidth::B16:	TextBuilder.AddValueHex(LeafView.AsHex16()); break;
+			case ELeafWidth::B32:	TextBuilder.AddValueHex(LeafView.AsHex32()); break;
+			case ELeafWidth::B64:	TextBuilder.AddValueHex(LeafView.AsHex64()); break;
+		}
 		break;
 	case ELeafType::Enum:
 		switch (LeafView.Leaf.Width)
 		{
-			case ELeafWidth::B8:	TextBuilder.AddLeafValue(LeafView.AsUnderlyingValue<uint8>()); break;
-			case ELeafWidth::B16:	TextBuilder.AddLeafValue(LeafView.AsUnderlyingValue<uint16>()); break;
-			case ELeafWidth::B32:	TextBuilder.AddLeafValue(LeafView.AsUnderlyingValue<uint32>()); break;
-			case ELeafWidth::B64:	TextBuilder.AddLeafValue(LeafView.AsUnderlyingValue<uint64>()); break;
+			case ELeafWidth::B8:	TextBuilder.AddValue(LeafView.AsUnderlyingValue<uint8>());  break;
+			case ELeafWidth::B16:	TextBuilder.AddValue(LeafView.AsUnderlyingValue<uint16>()); break;
+			case ELeafWidth::B32:	TextBuilder.AddValue(LeafView.AsUnderlyingValue<uint32>()); break;
+			case ELeafWidth::B64:	TextBuilder.AddValue(LeafView.AsUnderlyingValue<uint64>()); break;
 		}
 		break;
 	case ELeafType::Unicode:
 		switch (LeafView.Leaf.Width)
 		{
-			case ELeafWidth::B8:	TextBuilder.AddLeafValue(LeafView.AsChar8());	break;
-			case ELeafWidth::B16:	TextBuilder.AddLeafValue(LeafView.AsChar16());	break;
-			case ELeafWidth::B32:	TextBuilder.AddLeafValue(LeafView.AsChar32());	break;
+			case ELeafWidth::B8:	TextBuilder.AddValue(LeafView.AsChar8());	break;
+			case ELeafWidth::B16:	TextBuilder.AddValue(LeafView.AsChar16());	break;
+			case ELeafWidth::B32:	TextBuilder.AddValue(LeafView.AsChar32());	break;
 			case ELeafWidth::B64:	check(false);									break;
 		};
 	}
@@ -998,47 +1052,52 @@ void FMemberPrinter::PrintLeaves(FUnpackedLeafType Leaf, const FLeafRangeView& L
 	switch (Leaf.Type)
 	{
 	case ELeafType::Bool:
-		for (bool b : LeafRange.AsBools()) { TextBuilder.AddLeaf(b); } break;
+		for (bool b : LeafRange.AsBools()) { TextBuilder.AddItem(b); } break;
 	case ELeafType::IntS:
 		switch (Leaf.Width)
 		{
-			case ELeafWidth::B8:  for (const int8&  I : LeafRange.AsS8s())  { TextBuilder.AddLeaf(I); } break;
-			case ELeafWidth::B16: for (const int16& I : LeafRange.AsS16s()) { TextBuilder.AddLeaf(I); } break;
-			case ELeafWidth::B32: for (const int32& I : LeafRange.AsS32s()) { TextBuilder.AddLeaf(I); } break;
-			case ELeafWidth::B64: for (const int64& I : LeafRange.AsS64s()) { TextBuilder.AddLeaf(I); } break;
+			case ELeafWidth::B8:  for (const int8&  I : LeafRange.AsS8s())  { TextBuilder.AddItem(I); } break;
+			case ELeafWidth::B16: for (const int16& I : LeafRange.AsS16s()) { TextBuilder.AddItem(I); } break;
+			case ELeafWidth::B32: for (const int32& I : LeafRange.AsS32s()) { TextBuilder.AddItem(I); } break;
+			case ELeafWidth::B64: for (const int64& I : LeafRange.AsS64s()) { TextBuilder.AddItem(I); } break;
 		}
 		break;
 	case ELeafType::IntU:
 		switch (Leaf.Width)
 		{
-			case ELeafWidth::B8:  for (const uint8&  U : LeafRange.AsU8s())  { TextBuilder.AddLeaf(U); } break;
-			case ELeafWidth::B16: for (const uint16& U : LeafRange.AsU16s()) { TextBuilder.AddLeaf(U); } break;
-			case ELeafWidth::B32: for (const uint32& U : LeafRange.AsU32s()) { TextBuilder.AddLeaf(U); } break;
-			case ELeafWidth::B64: for (const uint64& U : LeafRange.AsU64s()) { TextBuilder.AddLeaf(U); } break;
+			case ELeafWidth::B8:  for (const uint8&  U : LeafRange.AsU8s())  { TextBuilder.AddItem(U); } break;
+			case ELeafWidth::B16: for (const uint16& U : LeafRange.AsU16s()) { TextBuilder.AddItem(U); } break;
+			case ELeafWidth::B32: for (const uint32& U : LeafRange.AsU32s()) { TextBuilder.AddItem(U); } break;
+			case ELeafWidth::B64: for (const uint64& U : LeafRange.AsU64s()) { TextBuilder.AddItem(U); } break;
 		}
 		break;
 	case ELeafType::Float:
 		if (Leaf.Width == ELeafWidth::B32)
 		{
-			for (const float& f : LeafRange.AsFloats()) { TextBuilder.AddLeaf(f); }
+			for (const float& f : LeafRange.AsFloats()) { TextBuilder.AddItem(f); }
 		}
 		else
 		{
 			check(Leaf.Width == ELeafWidth::B64);
-			for (const double& d : LeafRange.AsDoubles()) { TextBuilder.AddLeaf(d); }
+			for (const double& d : LeafRange.AsDoubles()) { TextBuilder.AddItem(d); }
 		}
 		break;
 	case ELeafType::Hex:
-		// PP-TEXT: Implement AddLeaf(Hex)
-		check(Leaf.Type != ELeafType::Hex);
+		switch (Leaf.Width)
+		{
+			case ELeafWidth::B8:  for (const uint8&  U : LeafRange.AsHex8s())  { TextBuilder.AddItemHex(U); } break;
+			case ELeafWidth::B16: for (const uint16& U : LeafRange.AsHex16s()) { TextBuilder.AddItemHex(U); } break;
+			case ELeafWidth::B32: for (const uint32& U : LeafRange.AsHex32s()) { TextBuilder.AddItemHex(U); } break;
+			case ELeafWidth::B64: for (const uint64& U : LeafRange.AsHex64s()) { TextBuilder.AddItemHex(U); } break;
+		}
 		break;
 	case ELeafType::Enum:
 		switch (Leaf.Width)
 		{
-			case ELeafWidth::B8:  for (const uint8&  U : LeafRange.AsUnderlyingValues<uint8>())  { TextBuilder.AddLeaf(U); } break;
-			case ELeafWidth::B16: for (const uint16& U : LeafRange.AsUnderlyingValues<uint16>()) { TextBuilder.AddLeaf(U); } break;
-			case ELeafWidth::B32: for (const uint32& U : LeafRange.AsUnderlyingValues<uint32>()) { TextBuilder.AddLeaf(U); } break;
-			case ELeafWidth::B64: for (const uint64& U : LeafRange.AsUnderlyingValues<uint64>()) { TextBuilder.AddLeaf(U); } break;
+			case ELeafWidth::B8:  for (const uint8&  U : LeafRange.AsUnderlyingValues<uint8>())  { TextBuilder.AddItem(U); } break;
+			case ELeafWidth::B16: for (const uint16& U : LeafRange.AsUnderlyingValues<uint16>()) { TextBuilder.AddItem(U); } break;
+			case ELeafWidth::B32: for (const uint32& U : LeafRange.AsUnderlyingValues<uint32>()) { TextBuilder.AddItem(U); } break;
+			case ELeafWidth::B64: for (const uint64& U : LeafRange.AsUnderlyingValues<uint64>()) { TextBuilder.AddItem(U); } break;
 		}
 		break;
 	case ELeafType::Unicode:
@@ -1160,11 +1219,11 @@ static void AddUnicodeRangeLeaf(FYamlBuilder& TextBuilder, FOptionalMemberId Id,
 	check(DstEnd - Dst == DstLen);
 	if (Id)
 	{
-		TextBuilder.AddLeafValue(FUtf8StringView(Dst, DstLen));
+		TextBuilder.AddValue(FUtf8StringView(Dst, DstLen));
 	}
 	else
 	{
-		TextBuilder.AddLeaf(FUtf8StringView(Dst, DstLen));
+		TextBuilder.AddItem(FUtf8StringView(Dst, DstLen));
 	}
 }
 
@@ -1174,7 +1233,7 @@ void FMemberPrinter::PrintUnicodeRangeAsLeaf(FOptionalMemberId Id, FRangeType Ra
 
 	if (Id)
 	{
-		TextBuilder.AddLeafId(*FPrintId<128>(Ids, Id));
+		TextBuilder.AddId(*FPrintId<128>(Ids, Id));
 	}
 
 	const FLeafRangeView LeafRange = RangeView.AsLeaves();
diff --git a/Engine/Plugins/Experimental/PlainProps/Source/Public/PlainPropsBuild.h b/Engine/Plugins/Experimental/PlainProps/Source/Public/PlainPropsBuild.h
--- a/Engine/Plugins/Experimental/PlainProps/Source/Public/PlainPropsBuild.h
+++ b/Engine/Plugins/Experimental/PlainProps/Source/Public/PlainPropsBuild.h
@@ -151,6 +151,12 @@ FMemberSchema MakeLeafRangeSchema()
 	return MakeLeafRangeSchema<T>(RangeSizeOf(SizeType{}));
 }
 
+template<UnsignedIntegral T>
+FMemberSchema MakeHexRangeSchema(ESizeType MaxSize)
+{
+	return { FMemberType(MaxSize), FMemberType(ELeafType::Hex, ReflectArithmetic<T>.Width), 1, NoId, nullptr };
+}
+
 template<Enumeration T>
 FMemberSchema MakeEnumRangeSchema(FEnumId Id, ESizeType MaxSize)
 {
@@ -215,6 +221,18 @@ template<Arithmetic T, typename SizeType>
 	return BuildLeafRange(Scratch, RangeSizeOf(SizeType{}), Values.GetData(), Values.Num());
 }
 
+template<UnsignedIntegral T, typename SizeType>
+[[nodiscard]] FTypedRange BuildHexRange(FScratchAllocator& Scratch, TConstArrayView<T, SizeType> Values)
+{
+	return { MakeHexRangeSchema<T>(RangeSizeOf(SizeType{})), CloneLeaves(Scratch, Values.Num(), Values.GetData(), sizeof(T)) };
+}
+
+template<UnsignedIntegral T>
+[[nodiscard]] FTypedRange BuildHexRange(FScratchAllocator& Scratch, ESizeType SizeType, TConstArrayView<T> Values)
+{
+	return { MakeHexRangeSchema<T>(SizeType), CloneLeaves(Scratch, Values.Num(), Values.GetData(), sizeof(T)) };
+}
+
 template<Enumeration T, typename SizeType>
 [[nodiscard]] FTypedRange BuildEnumRange(FScratchAllocator& Scratch, FEnumId Enum, TConstArrayView<T, SizeType> Values)
 {
diff --git a/Engine/Plugins/Experimental/PlainProps/Source/Public/PlainPropsParse.h b/Engine/Plugins/Experimental/PlainProps/Source/Public/PlainPropsParse.h
--- a/Engine/Plugins/Experimental/PlainProps/Source/Public/PlainPropsParse.h
+++ b/Engine/Plugins/Experimental/PlainProps/Source/Public/PlainPropsParse.h
@@ -8,28 +8,40 @@
 namespace PlainProps
 {
 
-template<typename T>	PLAINPROPS_API TOptional<T>						Parse(FUtf8StringView String);
-
-template<>				PLAINPROPS_API TOptional<ESizeType>				Parse(FUtf8StringView String);
-template<>				PLAINPROPS_API TOptional<ELeafWidth>			Parse(FUtf8StringView String);
-
-template<>				PLAINPROPS_API TOptional<FUnpackedLeafType>		Parse(FUtf8StringView String);
-
-template<>				PLAINPROPS_API TOptional<bool>					Parse(FUtf8StringView String);
-template<>				PLAINPROPS_API TOptional<int8> 					Parse(FUtf8StringView String);
-template<>				PLAINPROPS_API TOptional<int16>					Parse(FUtf8StringView String);
-template<>				PLAINPROPS_API TOptional<int32>					Parse(FUtf8StringView String);
-template<>				PLAINPROPS_API TOptional<int64>					Parse(FUtf8StringView String);
-template<>				PLAINPROPS_API TOptional<uint8> 				Parse(FUtf8StringView String);
-template<>				PLAINPROPS_API TOptional<uint16>				Parse(FUtf8StringView String);
-template<>				PLAINPROPS_API TOptional<uint32>				Parse(FUtf8StringView String);
-template<>				PLAINPROPS_API TOptional<uint64>				Parse(FUtf8StringView String);
-template<>				PLAINPROPS_API TOptional<float>					Parse(FUtf8StringView String);
-template<>				PLAINPROPS_API TOptional<double>				Parse(FUtf8StringView String);
-template<>				PLAINPROPS_API TOptional<char>					Parse(FUtf8StringView String);
-template<>				PLAINPROPS_API TOptional<char8_t>				Parse(FUtf8StringView String);
-template<>				PLAINPROPS_API TOptional<char16_t>				Parse(FUtf8StringView String);
-template<>				PLAINPROPS_API TOptional<char32_t>				Parse(FUtf8StringView String);
+template<typename T> TOptional<T> Parse(FUtf8StringView String)
+{
+	static_assert(!sizeof(T), "Unsupported type for Parse");
+}
+
+template<> PLAINPROPS_API TOptional<ESizeType>			Parse(FUtf8StringView String);
+template<> PLAINPROPS_API TOptional<ELeafWidth>			Parse(FUtf8StringView String);
+
+template<> PLAINPROPS_API TOptional<FUnpackedLeafType>	Parse(FUtf8StringView String);
+
+template<> PLAINPROPS_API TOptional<bool>				Parse(FUtf8StringView String);
+template<> PLAINPROPS_API TOptional<int8> 				Parse(FUtf8StringView String);
+template<> PLAINPROPS_API TOptional<int16>				Parse(FUtf8StringView String);
+template<> PLAINPROPS_API TOptional<int32>				Parse(FUtf8StringView String);
+template<> PLAINPROPS_API TOptional<int64>				Parse(FUtf8StringView String);
+template<> PLAINPROPS_API TOptional<uint8> 				Parse(FUtf8StringView String);
+template<> PLAINPROPS_API TOptional<uint16>				Parse(FUtf8StringView String);
+template<> PLAINPROPS_API TOptional<uint32>				Parse(FUtf8StringView String);
+template<> PLAINPROPS_API TOptional<uint64>				Parse(FUtf8StringView String);
+template<> PLAINPROPS_API TOptional<float>				Parse(FUtf8StringView String);
+template<> PLAINPROPS_API TOptional<double>				Parse(FUtf8StringView String);
+template<> PLAINPROPS_API TOptional<char>				Parse(FUtf8StringView String);
+template<> PLAINPROPS_API TOptional<char8_t>			Parse(FUtf8StringView String);
+template<> PLAINPROPS_API TOptional<char16_t>			Parse(FUtf8StringView String);
+template<> PLAINPROPS_API TOptional<char32_t>			Parse(FUtf8StringView String);
+
+template<typename T> TOptional<T> ParseHex(FUtf8StringView String)
+{
+	static_assert(!sizeof(T), "Unsupported type for ParseHex");
+}
+template<> PLAINPROPS_API TOptional<uint8>				ParseHex(FUtf8StringView String);
+template<> PLAINPROPS_API TOptional<uint16>				ParseHex(FUtf8StringView String);
+template<> PLAINPROPS_API TOptional<uint32>				ParseHex(FUtf8StringView String);
+template<> PLAINPROPS_API TOptional<uint64>				ParseHex(FUtf8StringView String);
 
 ///////////////////////////////////////////////////////////////////////////////
 
diff --git a/Engine/Plugins/Experimental/PlainProps/Source/Public/PlainPropsPrint.h b/Engine/Plugins/Experimental/PlainProps/Source/Public/PlainPropsPrint.h
--- a/Engine/Plugins/Experimental/PlainProps/Source/Public/PlainPropsPrint.h
+++ b/Engine/Plugins/Experimental/PlainProps/Source/Public/PlainPropsPrint.h
@@ -45,6 +45,15 @@ template<> PLAINPROPS_API void Print(FUtf8Builder& Out, char8_t Value);
 template<> PLAINPROPS_API void Print(FUtf8Builder& Out, char16_t Value);
 template<> PLAINPROPS_API void Print(FUtf8Builder& Out, char32_t Value);
 
+template<typename T> inline void PrintHex(FUtf8Builder& Out, T Value)
+{
+	static_assert(!sizeof(T), "Unsupported type for PrintHex");
+}
+template<> PLAINPROPS_API void PrintHex(FUtf8Builder& Out, uint8 Value);
+template<> PLAINPROPS_API void PrintHex(FUtf8Builder& Out, uint16 Value);
+template<> PLAINPROPS_API void PrintHex(FUtf8Builder& Out, uint32 Value);
+template<> PLAINPROPS_API void PrintHex(FUtf8Builder& Out, uint64 Value);
+
 ///////////////////////////////////////////////////////////////////////////////
 
 PLAINPROPS_API void PrintMemberSchema(FUtf8Builder& Out, const FIds& Ids, FMemberSchema Member);
diff --git a/Engine/Plugins/Experimental/PlainProps/Source/Public/PlainPropsRead.h b/Engine/Plugins/Experimental/PlainProps/Source/Public/PlainPropsRead.h
--- a/Engine/Plugins/Experimental/PlainProps/Source/Public/PlainPropsRead.h
+++ b/Engine/Plugins/Experimental/PlainProps/Source/Public/PlainPropsRead.h
@@ -215,6 +215,10 @@ struct FLeafView
 	FORCEINLINE uint64		AsU64() const		{ return As<uint64>(); }
 	FORCEINLINE double		AsDouble() const	{ return As<double>(); }
 	FORCEINLINE float		AsFloat() const		{ return As<float>(); }
+	FORCEINLINE uint8		AsHex8() const		{ return As<uint8,  {ELeafType::Hex, ELeafWidth::B8}>();  }
+	FORCEINLINE uint16		AsHex16() const		{ return As<uint16, {ELeafType::Hex, ELeafWidth::B16}>(); }
+	FORCEINLINE uint32		AsHex32() const		{ return As<uint32, {ELeafType::Hex, ELeafWidth::B32}>(); }
+	FORCEINLINE uint64		AsHex64() const		{ return As<uint64, {ELeafType::Hex, ELeafWidth::B64}>(); }
 	FORCEINLINE char8_t		AsChar8() const		{ return As<char8_t>(); }
 	FORCEINLINE char16_t	AsChar16() const	{ return As<char16_t>(); }
 	FORCEINLINE char32_t	AsChar32() const	{ return As<char32_t>(); }
@@ -505,6 +509,10 @@ public:
 	TRangeView<uint64>		AsU64s() const		{ return As<uint64>(); }
 	TRangeView<float>		AsFloats() const	{ return As<float>(); }
 	TRangeView<double>		AsDoubles() const	{ return As<double>(); }
+	TRangeView<uint8>		AsHex8s() const		{ return As<uint8,  {ELeafType::Hex, ELeafWidth::B8}>();  }
+	TRangeView<uint16>		AsHex16s() const	{ return As<uint16, {ELeafType::Hex, ELeafWidth::B16}>(); }
+	TRangeView<uint32>		AsHex32s() const	{ return As<uint32, {ELeafType::Hex, ELeafWidth::B32}>(); }
+	TRangeView<uint64>		AsHex64s() const	{ return As<uint64, {ELeafType::Hex, ELeafWidth::B64}>(); }
 	TRangeView<char8_t>		AsUtf8() const		{ return As<char8_t>(); }
 	TRangeView<char16_t>	AsUtf16() const		{ return As<char16_t>(); }
 	TRangeView<char32_t>	AsUtf32() const		{ return As<char32_t>(); }
diff --git a/Engine/Plugins/Experimental/PlainProps/Source/Tests/PlainPropsTests.cpp b/Engine/Plugins/Experimental/PlainProps/Source/Tests/PlainPropsTests.cpp
--- a/Engine/Plugins/Experimental/PlainProps/Source/Tests/PlainPropsTests.cpp
+++ b/Engine/Plugins/Experimental/PlainProps/Source/Tests/PlainPropsTests.cpp
@@ -710,6 +710,77 @@ TEST_CASE_NAMED(FPlainPropsReadWriteTest, "System::Core::Serialization::PlainPro
 		});
 	}
 
+	SECTION("Hex")
+	{
+		static constexpr std::initializer_list<const char*> MemberNames =  {"Hex8", "Hex16", "Hex32", "Hex64"};
+
+		TestSerialize([](FTestBatchBuilder& Batch, FScratchAllocator& Scratch)
+		{
+			FMemberSpec MemberTypes[] = {SpecHex8, SpecHex16, SpecHex32, SpecHex64};
+			FDeclId Id = Batch.DeclareStruct("Test", "Hex", MemberNames, MemberTypes, EMemberPresence::AllowSparse);
+
+			FMemberBuilder Misc, Mins, Maxs, Some;
+
+			Misc.AddHex(Batch.NameMember("Hex8"),	uint8(8));
+			Misc.AddHex(Batch.NameMember("Hex16"),	uint16(16));
+			Misc.AddHex(Batch.NameMember("Hex32"),	uint32(32));
+			Misc.AddHex(Batch.NameMember("Hex64"),	uint64(64));
+			
+			Mins.AddHex(Batch.NameMember("Hex8"),	std::numeric_limits< uint8>::min());
+			Mins.AddHex(Batch.NameMember("Hex16"),	std::numeric_limits<uint16>::min());
+			Mins.AddHex(Batch.NameMember("Hex32"),	std::numeric_limits<uint32>::min());
+			Mins.AddHex(Batch.NameMember("Hex64"),	std::numeric_limits<uint64>::min());
+
+			Maxs.AddHex(Batch.NameMember("Hex8"),	std::numeric_limits< uint8>::max());
+			Maxs.AddHex(Batch.NameMember("Hex16"),	std::numeric_limits<uint16>::max());
+			Maxs.AddHex(Batch.NameMember("Hex32"),	std::numeric_limits<uint32>::max());
+			Maxs.AddHex(Batch.NameMember("Hex64"),	std::numeric_limits<uint64>::max());
+
+			Some.AddHex(Batch.NameMember("Hex32"),	uint32(0));
+			
+			Batch.AddObject(Id, MoveTemp(Misc));
+			Batch.AddObject(Id, MoveTemp(Mins));
+			Batch.AddObject(Id, MoveTemp(Maxs));
+			Batch.AddObject(Id, MoveTemp(Some));
+		}, 
+		[](TConstArrayView<FStructView> Objects, const FTestNameReader& Names)
+		{
+			for (const FStructView& Object : Objects.Left(3))
+			{
+				FTestMemberReader Members(Object);
+				for (const char* MemberName : MemberNames)
+				{
+					CHECK(Members.HasMore());
+					CHECK(Names[Members.PeekName()] == MemberName);
+					CHECK(Members.PeekKind() == EMemberKind::Leaf);
+					(void)Members.GrabLeaf();
+				}
+			}
+			
+			FTestMemberReader Misc(Objects[0]);
+			CHECK(Misc.GrabLeaf().AsHex8()	== uint8(8));
+			CHECK(Misc.GrabLeaf().AsHex16()	== uint16(16));
+			CHECK(Misc.GrabLeaf().AsHex32()	== uint32(32));
+			CHECK(Misc.GrabLeaf().AsHex64()	== uint64(64));
+
+			FTestMemberReader Mins(Objects[1]);
+			CHECK(Mins.GrabLeaf().AsHex8()	== std::numeric_limits< uint8>::min());
+			CHECK(Mins.GrabLeaf().AsHex16()	== std::numeric_limits<uint16>::min());
+			CHECK(Mins.GrabLeaf().AsHex32()	== std::numeric_limits<uint32>::min());
+			CHECK(Mins.GrabLeaf().AsHex64()	== std::numeric_limits<uint64>::min());
+
+			FTestMemberReader Maxs(Objects[2]);
+			CHECK(Maxs.GrabLeaf().AsHex8()	== std::numeric_limits< uint8>::max());
+			CHECK(Maxs.GrabLeaf().AsHex16()	== std::numeric_limits<uint16>::max());
+			CHECK(Maxs.GrabLeaf().AsHex32()	== std::numeric_limits<uint32>::max());
+			CHECK(Maxs.GrabLeaf().AsHex64()	== std::numeric_limits<uint64>::max());
+
+			FTestMemberReader Some(Objects[3]);
+			CHECK(Names[Some.PeekName()] == "Hex32");
+			CHECK(Some.GrabLeaf().AsHex32() == 0);
+		});
+	}
+
 	SECTION("Unicode")
 	{
 		TestSerialize([](FTestBatchBuilder& Batch, FScratchAllocator& Scratch)
@@ -1039,10 +1110,11 @@ TEST_CASE_NAMED(FPlainPropsReadWriteTest, "System::Core::Serialization::PlainPro
 			FEnumId Unused1	= Batch.DeclareEnum("Test", "Unused1", EEnumMode::Flat, ELeafWidth::B8, {"X"}, {0});
 			FEnumId Unused2	= Batch.DeclareEnum("Test", "Unused2", EEnumMode::Flat, ELeafWidth::B8, {"Y"}, {0});
 
-			static constexpr std::initializer_list<const char*> MemberNames = { "B0", "B1", "B8", "B9", "D0", "D3", "Hi", "E3", "E0" };
+			static constexpr std::initializer_list<const char*> MemberNames = { "B0", "B1", "B8", "B9", "D0", "D3", "H0", "E3", "E0", "Hi" };
 			FMemberSpec MemberTypes[] = {	S32Range(SpecBool), S32Range(SpecBool), S64Range(SpecBool), S64Range(SpecBool), 
-											S32Range(SpecF64), S32Range(SpecF64), S32Range(SpecUtf8), 
-											S32Range(Enum16(ABCD)), S32Range(Enum8(Unused1)) };
+											S32Range(SpecF64), S32Range(SpecF64),
+											S32Range(SpecHex16), S32Range(Enum16(ABCD)), S32Range(Enum8(Unused1)),
+											S32Range(SpecUtf8)};
 			FDeclId ObjectId = Batch.DeclareStruct("Test", "Object", MemberNames, MemberTypes, EMemberPresence::AllowSparse);
 		
 			FMemberBuilder Members;
@@ -1052,9 +1124,10 @@ TEST_CASE_NAMED(FPlainPropsReadWriteTest, "System::Core::Serialization::PlainPro
 			Members.AddRange(Batch.NameMember("B9"), BuildLeafRange(Scratch, TConstArrayView<bool,int64>({true, false, true, false, true, false, true, false, true})));
 			Members.AddRange(Batch.NameMember("D0"), BuildLeafRange(Scratch, TConstArrayView<double>()));
 			Members.AddRange(Batch.NameMember("D3"), BuildLeafRange(Scratch, MakeArrayView({DBL_MIN, 0.0, DBL_MAX})));
-			Members.AddRange(Batch.NameMember("Hi"), BuildLeafRange(Scratch, MakeArrayView(u8"Hello!")));
+			Members.AddRange(Batch.NameMember("H0"), BuildHexRange(Scratch,  TConstArrayView<uint16>({0x00, 0x01, 0xEF, 0xFF})));
 			Members.AddRange(Batch.NameMember("E3"), BuildEnumRange(Scratch, ABCD, MakeArrayView({EABCD::B, EABCD::A, EABCD::D})));
 			Members.AddRange(Batch.NameMember("E0"), BuildEnumRange(Scratch, Unused1, TConstArrayView<EUnused1>()));
+			Members.AddRange(Batch.NameMember("Hi"), BuildLeafRange(Scratch, MakeArrayView(u8"Hello!")));
 
 			Batch.AddObject(ObjectId, MoveTemp(Members));
 		}, 
@@ -1069,19 +1142,21 @@ TEST_CASE_NAMED(FPlainPropsReadWriteTest, "System::Core::Serialization::PlainPro
 			FLeafRangeView B9 = It.GrabRange().AsLeaves();
 			FLeafRangeView D0 = It.GrabRange().AsLeaves();
 			FLeafRangeView D3 = It.GrabRange().AsLeaves();
-			FLeafRangeView Hi = It.GrabRange().AsLeaves();
+			FLeafRangeView H0 = It.GrabRange().AsLeaves();
 			FLeafRangeView E3 = It.GrabRange().AsLeaves();
 			FLeafRangeView E0 = It.GrabRange().AsLeaves();
+			FLeafRangeView Hi = It.GrabRange().AsLeaves();
 			
 			CHECK(B0.Num() == 0);
-			CHECK(EqualItems(B1.AsBools(),	MakeArrayView({true})));
-			CHECK(EqualItems(B8.AsBools(),	MakeArrayView({false, true, false, true, false, true, false, true})));
-			CHECK(EqualItems(B9.AsBools(),	MakeArrayView({true, false, true, false, true, false, true, false, true})));
-			CHECK(EqualItems(D0.AsDoubles(), TConstArrayView<double>()));
-			CHECK(EqualItems(D3.AsDoubles(), MakeArrayView({DBL_MIN, 0.0, DBL_MAX})));
-			CHECK(EqualItems(Hi.AsUtf8(), u8"Hello!"));
-			CHECK(EqualItems(E3.As<EABCD>(), MakeArrayView({EABCD::B, EABCD::A, EABCD::D})));
-			CHECK(EqualItems(E0.As<EUnused1>(), TConstArrayView<EUnused1>()));
+			CHECK(EqualItems(B1.AsBools(),		MakeArrayView({true})));
+			CHECK(EqualItems(B8.AsBools(),		MakeArrayView({false, true, false, true, false, true, false, true})));
+			CHECK(EqualItems(B9.AsBools(),		MakeArrayView({true, false, true, false, true, false, true, false, true})));
+			CHECK(EqualItems(D0.AsDoubles(),	TConstArrayView<double>()));
+			CHECK(EqualItems(D3.AsDoubles(),	MakeArrayView({DBL_MIN, 0.0, DBL_MAX})));
+			CHECK(EqualItems(H0.AsHex16s(),		TConstArrayView<uint16>({0x00, 0x01, 0xEF, 0xFF})));
+			CHECK(EqualItems(E3.As<EABCD>(),	MakeArrayView({EABCD::B, EABCD::A, EABCD::D})));
+			CHECK(EqualItems(E0.As<EUnused1>(),	TConstArrayView<EUnused1>()));
+			CHECK(EqualItems(Hi.AsUtf8(),		u8"Hello!"));
 		});
 	}
 
diff --git a/Engine/Plugins/Experimental/PlainPropsEngine/Source/Private/PlainPropsEngineBindings.cpp b/Engine/Plugins/Experimental/PlainPropsEngine/Source/Private/PlainPropsEngineBindings.cpp
--- a/Engine/Plugins/Experimental/PlainPropsEngine/Source/Private/PlainPropsEngineBindings.cpp
+++ b/Engine/Plugins/Experimental/PlainPropsEngine/Source/Private/PlainPropsEngineBindings.cpp
@@ -58,7 +58,7 @@ struct FTexture2DCook : ICustomBinding
 
 	FTexture2DCook(TPropertySpecifier<5>& Spec)
 	: TextureAddressId(IndexEngineEnum("TextureAddress"))
-	, IntPointId(IndexStruct(GUE.Scopes.Core, "IntPoint"))
+	, IntPointId(IndexStruct(GUE.Scopes.CoreUObject, "IntPoint"))
 	, PlatformTextureDataId(IndexEngineStruct("PlatformTextureData"))
 	, MemberIds{GUE.Names.NameMember("AddressX"),
 				GUE.Names.NameMember("AddressY"),
@@ -148,4 +148,4 @@ void CustomBindEngineTypes(EBindMode Mode)
 	}
 }
 
-} // namespace PlainProps::UE
\ No newline at end of file
+} // namespace PlainProps::UE
diff --git a/Engine/Plugins/Experimental/PlainPropsUObject/Source/Private/PlainPropsUObjectBindings.cpp b/Engine/Plugins/Experimental/PlainPropsUObject/Source/Private/PlainPropsUObjectBindings.cpp
--- a/Engine/Plugins/Experimental/PlainPropsUObject/Source/Private/PlainPropsUObjectBindings.cpp
+++ b/Engine/Plugins/Experimental/PlainPropsUObject/Source/Private/PlainPropsUObjectBindings.cpp
@@ -250,8 +250,7 @@ const void* FDefaultStructs::Get(FBindId Id)
 ////////////////////////////////////////////////////////////////////////////////////////////////
 
 FCommonScopeIds::FCommonScopeIds(TIdIndexer<FSensitiveName>& Names)
-: Core(						Names.MakeScope("/Script/Core"))
-, CoreUObject(				Names.MakeScope("/Script/CoreUObject"))
+: CoreUObject(				Names.MakeScope("/Script/CoreUObject"))
 , Engine(					Names.MakeScope("/Script/Engine"))
 {}
 
@@ -274,13 +273,13 @@ FCommonTypenameIds::FCommonTypenameIds(TIdIndexer<FSensitiveName>& Names)
 {}
 
 FCommonEnumIds::FCommonEnumIds(const FCommonScopeIds& Scopes, TIdIndexer<FSensitiveName>& Names)
-: PixelFormat(Names.IndexEnum({Scopes.Core, Names.MakeTypename("PixelFormat")}))
+: PixelFormat(Names.IndexEnum({Scopes.CoreUObject, Names.MakeTypename("PixelFormat")}))
 {}
 
 FCommonStructIds::FCommonStructIds(const FCommonScopeIds& Scopes, TIdIndexer<FSensitiveName>& Names)
-: Name(						Names.IndexStruct({Scopes.Core, Names.MakeTypename("Name")}))
-, Text(						Names.IndexStruct({Scopes.Core, Names.MakeTypename("Text")}))
-, Guid(						Names.IndexStruct({Scopes.Core, Names.MakeTypename("Guid")}))
+: Name(						Names.IndexStruct({Scopes.CoreUObject, Names.MakeTypename("Name")}))
+, Text(						Names.IndexStruct({Scopes.CoreUObject, Names.MakeTypename("Text")}))
+, Guid(						Names.IndexStruct({Scopes.CoreUObject, Names.MakeTypename("Guid")}))
 , FieldPath(				Names.IndexStruct({Scopes.CoreUObject, Names.MakeTypename("FieldPath")}))
 , SoftObjectPath(			Names.IndexStruct({Scopes.CoreUObject, Names.MakeTypename("SoftObjectPath")}))
 , ClassPtr(					Names.IndexStruct({Scopes.CoreUObject, Names.MakeTypename("ClassPtr")}))
@@ -1908,7 +1907,7 @@ bool DiffMap(const BindingType& Binding, const ScriptType& A, const ScriptType&
 
 			// Temp comment: Fails for Lyra -loadmaps since FInstancedStruct is unbound,
 			//				 so all instanceed structs are considered equal
-			check(Binding.HasItem(B, PairA));
+			checkSlow(Binding.HasItem(B, PairA));
 		}
 	}
 
@@ -2236,9 +2235,9 @@ public:
 
 		// Index custom delta binding struct name
 		FBothType Param = Elem.IndexParameterName(GUE.Names);
-		FType BindType = FType{ GUE.Scopes.Core, FTypenameId(GUE.Names.MakeParametricTypeId(GUE.Typenames.Set, MakeArrayView(&Param.BindType, 1))) };
+		FType BindType = FType{ GUE.Scopes.CoreUObject, FTypenameId(GUE.Names.MakeParametricTypeId(GUE.Typenames.Set, MakeArrayView(&Param.BindType, 1))) };
 		FType DeclType = Param.IsLowered()
-						 ? FType{ GUE.Scopes.Core, FTypenameId(GUE.Names.MakeParametricTypeId(GUE.Typenames.Set, MakeArrayView(&Param.DeclType, 1))) }
+						 ? FType{ GUE.Scopes.CoreUObject, FTypenameId(GUE.Names.MakeParametricTypeId(GUE.Typenames.Set, MakeArrayView(&Param.DeclType, 1))) }
 						 : BindType;
 		FBindId BindId = GUE.Names.IndexBindId(BindType);
 		FDeclId DeclId = Param.IsLowered() ? GUE.Names.IndexDeclId(DeclType) : LowerCast(BindId);
@@ -2592,8 +2591,8 @@ class FMapBindings
 			checkf(BothKey.IsLowered() || BothValue.IsLowered(), TEXT("Key or Value is range-bound and should be type-erased / lowered"));
 			FType BindParams[2] = {BothKey.BindType, BothValue.BindType};
 			FType DeclParams[2] = {BothKey.DeclType, BothValue.DeclType};
-			FType BindType = { GUE.Scopes.Core, FTypenameId(GUE.Names.MakeParametricTypeId(GUE.Typenames.Map, BindParams)) };
-			FType DeclType = { GUE.Scopes.Core, FTypenameId(GUE.Names.MakeParametricTypeId(GUE.Typenames.Map, DeclParams)) };
+			FType BindType = { GUE.Scopes.CoreUObject, FTypenameId(GUE.Names.MakeParametricTypeId(GUE.Typenames.Map, BindParams)) };
+			FType DeclType = { GUE.Scopes.CoreUObject, FTypenameId(GUE.Names.MakeParametricTypeId(GUE.Typenames.Map, DeclParams)) };
 			BothPair = { GUE.Names.IndexBindId(BindType), GUE.Names.IndexDeclId(DeclType) };
 
 			FBindId PairId = Members.Pair.InnermostSchema.Get().AsStructBindId();
@@ -2617,7 +2616,7 @@ class FMapBindings
 
 			FParametricTypeId PairTypename = GUE.Names.Resolve(Members.Pair.InnermostSchema.Get().AsStruct()).Name.AsParametric();
 			TConstArrayView<FType> Params = GUE.Names.Resolve(PairTypename).GetParameters();
-			FType Type = { GUE.Scopes.Core, FTypenameId(GUE.Names.MakeParametricTypeId(GUE.Typenames.Map, Params)) };
+			FType Type = { GUE.Scopes.CoreUObject, FTypenameId(GUE.Names.MakeParametricTypeId(GUE.Typenames.Map, Params)) };
 			FDualStructId Dual{GUE.Names.IndexStruct(Type)};
 			BothPair = { Dual, Dual };
 		}
@@ -2733,8 +2732,8 @@ class FPairBindings
 		FBothType ValueParam = Pair.KV[1].IndexParameterName(GUE.Names);
 		FType BindParams[2] = { KeyParam.BindType, ValueParam.BindType };
 		FType DeclParams[2] = { KeyParam.DeclType, ValueParam.DeclType };
-		FType BindType = { GUE.Scopes.Core, FTypenameId(GUE.Names.MakeParametricTypeId(GUE.Typenames.Pair, BindParams)) };
-		FType DeclType = { GUE.Scopes.Core, FTypenameId(GUE.Names.MakeParametricTypeId(GUE.Typenames.Pair, DeclParams)) };
+		FType BindType = { GUE.Scopes.CoreUObject, FTypenameId(GUE.Names.MakeParametricTypeId(GUE.Typenames.Pair, BindParams)) };
+		FType DeclType = { GUE.Scopes.CoreUObject, FTypenameId(GUE.Names.MakeParametricTypeId(GUE.Typenames.Pair, DeclParams)) };
 		FBindId BindId = GUE.Names.IndexBindId(BindType);
 		FDeclId DeclId = GUE.Names.IndexDeclId(DeclType);
 
diff --git a/Engine/Plugins/Experimental/PlainPropsUObject/Source/Public/PlainPropsUObjectRuntime.h b/Engine/Plugins/Experimental/PlainPropsUObject/Source/Public/PlainPropsUObjectRuntime.h
--- a/Engine/Plugins/Experimental/PlainPropsUObject/Source/Public/PlainPropsUObjectRuntime.h
+++ b/Engine/Plugins/Experimental/PlainPropsUObject/Source/Public/PlainPropsUObjectRuntime.h
@@ -46,7 +46,6 @@ struct FCommonScopeIds
 {
 	explicit FCommonScopeIds(TIdIndexer<FSensitiveName>& Names);
 	
-	FScopeId					Core;
 	FScopeId					CoreUObject;
 	FScopeId					Engine; // Move out?
 };
@@ -265,14 +264,14 @@ template<> struct TCustomBind<::UE::FDynamicallyTypedValue> { using Type = UE::F
 template<> struct TCustomBind<FReferencePropertyValue> { using Type = UE::FReferencePropertyBinding; };
 template<> struct TCustomBind<FVerseFunction> { using Type = UE::FVerseFunctionBinding; };
 
-// Temporary way to tie certain types to /Script/Core scope
-template<> inline FScopeId IndexNamespaceId<UE::FRuntimeIds, TTypename<FIntPoint>>()		{ return UE::GUE.Scopes.Core; }
-template<> inline FScopeId IndexNamespaceId<UE::FRuntimeIds, TTypename<FQuat>>()			{ return UE::GUE.Scopes.Core; }
-template<> inline FScopeId IndexNamespaceId<UE::FRuntimeIds, TTypename<FTransform>>()		{ return UE::GUE.Scopes.Core; }
-template<> inline FScopeId IndexNamespaceId<UE::FRuntimeIds, TTypename<FVector>>()			{ return UE::GUE.Scopes.Core; }
-template<> inline FScopeId IndexNamespaceId<UE::FRuntimeIds, TTypename<FVector4>>()			{ return UE::GUE.Scopes.Core; }
-template<> inline FScopeId IndexNamespaceId<UE::FRuntimeIds, TTypename<FGuid>>()			{ return UE::GUE.Scopes.Core; }
-template<> inline FScopeId IndexNamespaceId<UE::FRuntimeIds, TTypename<FColor>>()			{ return UE::GUE.Scopes.Core; }
-template<> inline FScopeId IndexNamespaceId<UE::FRuntimeIds, TTypename<FLinearColor>>()		{ return UE::GUE.Scopes.Core; }
+// Temporary way to tie certain types to /Script/CoreUObject scope
+template<> inline FScopeId IndexNamespaceId<UE::FRuntimeIds, TTypename<FIntPoint>>()		{ return UE::GUE.Scopes.CoreUObject; }
+template<> inline FScopeId IndexNamespaceId<UE::FRuntimeIds, TTypename<FQuat>>()			{ return UE::GUE.Scopes.CoreUObject; }
+template<> inline FScopeId IndexNamespaceId<UE::FRuntimeIds, TTypename<FTransform>>()		{ return UE::GUE.Scopes.CoreUObject; }
+template<> inline FScopeId IndexNamespaceId<UE::FRuntimeIds, TTypename<FVector>>()			{ return UE::GUE.Scopes.CoreUObject; }
+template<> inline FScopeId IndexNamespaceId<UE::FRuntimeIds, TTypename<FVector4>>()			{ return UE::GUE.Scopes.CoreUObject; }
+template<> inline FScopeId IndexNamespaceId<UE::FRuntimeIds, TTypename<FGuid>>()			{ return UE::GUE.Scopes.CoreUObject; }
+template<> inline FScopeId IndexNamespaceId<UE::FRuntimeIds, TTypename<FColor>>()			{ return UE::GUE.Scopes.CoreUObject; }
+template<> inline FScopeId IndexNamespaceId<UE::FRuntimeIds, TTypename<FLinearColor>>()		{ return UE::GUE.Scopes.CoreUObject; }
 
 }
