PlainProps allow custom bindings to opt-in to memcpy loading

#rnx

diff --git a/Engine/Plugins/Experimental/PlainProps/Source/Private/PlainPropsLoad.cpp b/Engine/Plugins/Experimental/PlainProps/Source/Private/PlainPropsLoad.cpp
--- a/Engine/Plugins/Experimental/PlainProps/Source/Private/PlainPropsLoad.cpp
+++ b/Engine/Plugins/Experimental/PlainProps/Source/Private/PlainPropsLoad.cpp
@@ -15,12 +15,6 @@
 namespace PlainProps
 {
 
-struct FMemcpyLoadPlan
-{
-	uint32					Size;
-	uint32					Offset;
-};
-
 // Duplicated runtime FSchemaBinding with FStructSchemaId load ids instead of runtime FInnerId
 struct FSchemaLoadPlan
 {
diff --git a/Engine/Plugins/Experimental/PlainProps/Source/Private/PlainPropsUeCoreBindings.cpp b/Engine/Plugins/Experimental/PlainProps/Source/Private/PlainPropsUeCoreBindings.cpp
--- a/Engine/Plugins/Experimental/PlainProps/Source/Private/PlainPropsUeCoreBindings.cpp
+++ b/Engine/Plugins/Experimental/PlainProps/Source/Private/PlainPropsUeCoreBindings.cpp
@@ -175,12 +175,11 @@ void FGuidBinding::Save(FMemberBuilder& Dst, const FGuid& Src, const FGuid*, con
 	Dst.AddHex(MemberIds[1], Src.B);
 	Dst.AddHex(MemberIds[2], Src.C);
 	Dst.AddHex(MemberIds[3], Src.D);
-}
 
-void FGuidBinding::Load(FGuid& Dst, FStructLoadView Src, ECustomLoadMethod Method) const
-{
-	Src.Values.CheckSize(sizeof(FGuid));
-	FMemory::Memcpy(&Dst, Src.Values.Peek(), sizeof(FGuid));
+	static_assert(offsetof(FGuid, A) == 0);
+	static_assert(offsetof(FGuid, B) == 4);
+	static_assert(offsetof(FGuid, C) == 8);
+	static_assert(offsetof(FGuid, D) == 12);
 }
 
 //////////////////////////////////////////////////////////////////////////
@@ -191,16 +190,11 @@ void FColorBinding::Save(FMemberBuilder& Dst, const FColor& Src, const FColor*,
 	Dst.Add(MemberIds[1], Src.G);
 	Dst.Add(MemberIds[2], Src.R);
 	Dst.Add(MemberIds[3], Src.A);
-}
 
-void FColorBinding::Load(FColor& Dst, FStructLoadView Src, ECustomLoadMethod Method) const
-{
 	static_assert(offsetof(FColor, B) == 0);
 	static_assert(offsetof(FColor, G) == 1);
 	static_assert(offsetof(FColor, R) == 2);
 	static_assert(offsetof(FColor, A) == 3);
-	Src.Values.CheckSize(sizeof(FColor));
-	FMemory::Memcpy(&Dst, Src.Values.Peek(), sizeof(FColor));
 }
 
 //////////////////////////////////////////////////////////////////////////
@@ -211,16 +205,11 @@ void FLinearColorBinding::Save(FMemberBuilder& Dst, const FLinearColor& Src, con
 	Dst.Add(MemberIds[1], Src.G);
 	Dst.Add(MemberIds[2], Src.B);
 	Dst.Add(MemberIds[3], Src.A);
-}
 
-void FLinearColorBinding::Load(FLinearColor& Dst, FStructLoadView Src, ECustomLoadMethod Method) const
-{
 	static_assert(offsetof(FLinearColor, R) == 0);
 	static_assert(offsetof(FLinearColor, G) == 4);
 	static_assert(offsetof(FLinearColor, B) == 8);
 	static_assert(offsetof(FLinearColor, A) == 12);
-	Src.Values.CheckSize(sizeof(FLinearColor));
-	FMemory::Memcpy(&Dst, Src.Values.Peek(), sizeof(FLinearColor));
 }
 
 } // namespace PlainProps::UE
diff --git a/Engine/Plugins/Experimental/PlainProps/Source/Public/PlainPropsBind.h b/Engine/Plugins/Experimental/PlainProps/Source/Public/PlainPropsBind.h
--- a/Engine/Plugins/Experimental/PlainProps/Source/Public/PlainPropsBind.h
+++ b/Engine/Plugins/Experimental/PlainProps/Source/Public/PlainPropsBind.h
@@ -26,6 +26,7 @@ struct FDiffContext;
 class FIdIndexerBase;
 struct FLoadBatch;
 class FMemberBuilder;
+struct FMemcpyLoadPlan;
 class FRangeLoader;
 struct FSchemaBatch;
 class FScratchAllocator;
@@ -267,8 +268,12 @@ struct ICustomBinding
 	virtual void	LoadCustom(void* Dst, FStructLoadView Src, ECustomLoadMethod Method) const = 0;
 	virtual bool	DiffCustom(const void* StructA, const void* StructB, const FBindContext& Ctx) const = 0;
 
+	// Overload if struct should load via a single memcpy call instead of LoadCustom
+	virtual void	PlanCustom(FMemcpyLoadPlan& Out) const {}
+
 	// Overload to track the first diffing member, see FDiffPath
 	PLAINPROPS_API virtual bool	DiffCustom(const void* StructA, const void* StructB, FDiffContext& Ctx) const;
+	
 };
 
 struct FInnersHandle
@@ -920,6 +925,7 @@ struct TScopedStructBinding : FBothStructId, CustomBinding
 	using Ids = typename Runtime::Ids;
 	using Typename = TCustomTypename<CustomBinding>::Type;
 	static constexpr uint32 N = sizeof(CustomBinding::MemberIds) / sizeof(CustomBinding::MemberIds[0]);
+	static constexpr bool bFastLoad = requires(FMemcpyLoadPlan& O) { CustomBinding::Plan(O); };
 
 	explicit TScopedStructBinding(FType DeclName = IndexStructName<Ids, ETypename::Decl, Typename>())
 	: TScopedStructBinding(IndexStructBothId<Ids, Typename>(DeclName))
@@ -953,7 +959,14 @@ struct TScopedStructBinding : FBothStructId, CustomBinding
 
 	virtual void LoadCustom(void* Dst, TDependent<Ids>::StructLoadView Src, ECustomLoadMethod Method) const override 
 	{
-		CustomBinding::Load(*static_cast<T*>(Dst), Src, Method);
+		if constexpr (bFastLoad)
+		{
+			unimplemented();
+		}
+		else
+		{
+			CustomBinding::Load(*static_cast<T*>(Dst), Src, Method);
+		}
 	}
 	
 	virtual bool DiffCustom(const void* A, const void* B, const FBindContext& Ctx) const override
@@ -965,6 +978,14 @@ struct TScopedStructBinding : FBothStructId, CustomBinding
 	{
 		return CustomBinding::Diff(*static_cast<const T*>(A), *static_cast<const T*>(B), Ctx);
 	}
+
+	virtual void PlanCustom(FMemcpyLoadPlan& Out) const override
+	{
+		if constexpr (bFastLoad)
+		{
+			CustomBinding::Plan(Out);
+		}
+	}
 };
 
 // Specialization for schemabound structs
diff --git a/Engine/Plugins/Experimental/PlainProps/Source/Public/PlainPropsLoadMember.h b/Engine/Plugins/Experimental/PlainProps/Source/Public/PlainPropsLoadMember.h
--- a/Engine/Plugins/Experimental/PlainProps/Source/Public/PlainPropsLoadMember.h
+++ b/Engine/Plugins/Experimental/PlainProps/Source/Public/PlainPropsLoadMember.h
@@ -13,6 +13,12 @@ class FRangeLoader;
 class FRangeBinding;
 class FStructRangeLoadIterator;
 
+struct FMemcpyLoadPlan
+{
+	uint32					Size;
+	uint32					Offset;
+};
+
 struct FSchemaLoadHandle
 {
 	FStructSchemaId						LoadId;
diff --git a/Engine/Plugins/Experimental/PlainProps/Source/Public/PlainPropsUeCoreBindings.h b/Engine/Plugins/Experimental/PlainProps/Source/Public/PlainPropsUeCoreBindings.h
--- a/Engine/Plugins/Experimental/PlainProps/Source/Public/PlainPropsUeCoreBindings.h
+++ b/Engine/Plugins/Experimental/PlainProps/Source/Public/PlainPropsUeCoreBindings.h
@@ -666,8 +666,8 @@ struct FGuidBinding : ICustomBinding
 	}
 
 	PLAINPROPS_API void	Save(FMemberBuilder& Dst, const FGuid& Src, const FGuid* Default, const FSaveContext&) const;
-	PLAINPROPS_API void	Load(FGuid& Dst, FStructLoadView Src, ECustomLoadMethod Method) const;
-	inline static  bool Diff(FGuid A, FGuid B, const FBindContext&) { return A != B; }
+	inline static void	Plan(FMemcpyLoadPlan& Out) { Out.Size = sizeof(Type); }
+	inline static bool	Diff(FGuid A, FGuid B, const FBindContext&) { return A != B; }
 };
 
 //////////////////////////////////////////////////////////////////////////
@@ -685,8 +685,8 @@ struct FColorBinding : ICustomBinding
 	}
 	
 	PLAINPROPS_API void	Save(FMemberBuilder& Dst, const FColor& Src, const FColor* Default, const FSaveContext&) const;
-	PLAINPROPS_API void	Load(FColor& Dst, FStructLoadView Src, ECustomLoadMethod Method) const;
-	inline static  bool Diff(FColor A, FColor B, const FBindContext&) { return A != B; }
+	inline static void	Plan(FMemcpyLoadPlan& Out) { Out.Size = sizeof(Type); }
+	inline static bool	Diff(FColor A, FColor B, const FBindContext&) { return A != B; }
 };
 
 //////////////////////////////////////////////////////////////////////////
@@ -704,8 +704,8 @@ struct FLinearColorBinding : ICustomBinding
 	}
 	
 	PLAINPROPS_API void	Save(FMemberBuilder& Dst, const FLinearColor& Src, const FLinearColor* Default, const FSaveContext&) const;
-	PLAINPROPS_API void	Load(FLinearColor& Dst, FStructLoadView Src, ECustomLoadMethod Method) const;
-	inline static  bool Diff(FLinearColor A, FLinearColor B, const FBindContext&) { return A != B; }
+	inline static void	Plan(FMemcpyLoadPlan& Out) { Out.Size = sizeof(Type); }
+	inline static bool	Diff(FLinearColor A, FLinearColor B, const FBindContext&) { return A != B; }
 };
 
 //////////////////////////////////////////////////////////////////////////
diff --git a/Engine/Plugins/Experimental/PlainPropsUObject/Source/Private/PlainPropsRoundtripTest.cpp b/Engine/Plugins/Experimental/PlainPropsUObject/Source/Private/PlainPropsRoundtripTest.cpp
--- a/Engine/Plugins/Experimental/PlainPropsUObject/Source/Private/PlainPropsRoundtripTest.cpp
+++ b/Engine/Plugins/Experimental/PlainPropsUObject/Source/Private/PlainPropsRoundtripTest.cpp
@@ -80,10 +80,10 @@ struct FMemoryPropertyBatch
 		static_assert(sizeof(In) == sizeof(uint64));
 		Out.Add(GUE.Members.Id, reinterpret_cast<const uint64&>(In));
 	}
-
-	static void Load(FObjectHandle& Out, FStructLoadView In)
+	
+	static constexpr void Plan(FObjectHandle*, FMemcpyLoadPlan& Out)
 	{
-		LoadSole<uint64>(&Out, In);
+		Out.Size = sizeof(FObjectHandle);
 	}
 
 	static FSoleMemberSpec Spec(FWeakObjectPtr*, FDeclId Id)
@@ -98,9 +98,9 @@ struct FMemoryPropertyBatch
 		Out.Add(GUE.Members.Id, reinterpret_cast<const uint64&>(In));
 	}
 
-	static void Load(FWeakObjectPtr& Out, FStructLoadView In)
+	static constexpr void Plan(FWeakObjectPtr*, FMemcpyLoadPlan& Out)
 	{
-		LoadSole<uint64>(&Out, In);
+		Out.Size = sizeof(FWeakObjectPtr);
 	}
 
 	static FSoleMemberSpec Spec(FSoftObjectPtr*, FDeclId Id)
@@ -148,6 +148,8 @@ inline bool DiffProperty(const FLazyObjectPtr& A, const FLazyObjectPtr& B) { ret
 template<class Type, class BatchType>
 struct TCustomPropertyBinding final : ICustomBinding
 {
+	static constexpr bool bFastLoad = requires(Type* Tag, FMemcpyLoadPlan& Out) { BatchType::Plan(Tag, Out); };
+
 	TCustomPropertyBinding(BatchType& InBatch) : Batch(InBatch) {}
 
 	BatchType& Batch;
@@ -162,13 +164,28 @@ struct TCustomPropertyBinding final : ICustomBinding
 
 	virtual void LoadCustom(void* Dst, FStructLoadView Src, ECustomLoadMethod) const override
 	{
-		Batch.Load(*static_cast<Type*>(Dst), Src);
+		if constexpr (bFastLoad)
+		{
+			unimplemented();
+		}
+		else
+		{
+			Batch.Load(*static_cast<Type*>(Dst), Src);
+		}
 	}
 
 	virtual bool DiffCustom(const void* A, const void* B, const FBindContext&) const override
 	{
 		return DiffProperty(*static_cast<const Type*>(A), *static_cast<const Type*>(B));
 	}
+
+	virtual void PlanCustom(FMemcpyLoadPlan& Out) const override
+	{
+		if constexpr (bFastLoad)
+		{
+			Batch.Plan((Type*)nullptr, Out);
+		}
+	}
 };
 
 template<class BatchType>
