PlainProps use new custom binding load plans

#rnx

diff --git a/Engine/Plugins/Experimental/PlainProps/Source/Private/PlainPropsLoad.cpp b/Engine/Plugins/Experimental/PlainProps/Source/Private/PlainPropsLoad.cpp
--- a/Engine/Plugins/Experimental/PlainProps/Source/Private/PlainPropsLoad.cpp
+++ b/Engine/Plugins/Experimental/PlainProps/Source/Private/PlainPropsLoad.cpp
@@ -801,11 +801,20 @@ private:
 		{
 			if (const ICustomBinding* Custom = Customs.FindStruct(Mapping.Id))
 			{
-				return FLoadStructPlan(*Custom);
+				FMemcpyLoadPlan Memcpy;
+				Custom->PlanCustom(Memcpy);
+				return Memcpy.Size ? FLoadStructPlan(Memcpy) : FLoadStructPlan(*Custom);
 			}
 		}
 		else if (const ICustomBinding* Custom = Customs.FindStruct(Mapping.Id, /* out */ LoweredInners))
 		{
+			FMemcpyLoadPlan Memcpy;
+			Custom->PlanCustom(Memcpy);
+			if (Memcpy.Size)
+			{
+				return FLoadStructPlan(Memcpy);
+			}
+
 			TArray<FInnerStruct, TInlineAllocator<8>> TmpSubset;
 			TConstArrayView<FInnerStruct> LoweredMembers = GetLoweredMembers(LoweredInners, From.GetMemberNames(), /* out */ TmpSubset);
 
diff --git a/Engine/Plugins/Experimental/PlainProps/Source/Public/PlainPropsBind.h b/Engine/Plugins/Experimental/PlainProps/Source/Public/PlainPropsBind.h
--- a/Engine/Plugins/Experimental/PlainProps/Source/Public/PlainPropsBind.h
+++ b/Engine/Plugins/Experimental/PlainProps/Source/Public/PlainPropsBind.h
@@ -268,7 +268,7 @@ struct ICustomBinding
 	virtual void	LoadCustom(void* Dst, FStructLoadView Src, ECustomLoadMethod Method) const = 0;
 	virtual bool	DiffCustom(const void* StructA, const void* StructB, const FBindContext& Ctx) const = 0;
 
-	// Overload if struct should load via a single memcpy call instead of LoadCustom
+	// Overload to load via a single memcpy call instead of LoadCustom, Out.Offset is 0 by default
 	virtual void	PlanCustom(FMemcpyLoadPlan& Out) const {}
 
 	// Overload to track the first diffing member, see FDiffPath
diff --git a/Engine/Plugins/Experimental/PlainProps/Source/Public/PlainPropsLoadMember.h b/Engine/Plugins/Experimental/PlainProps/Source/Public/PlainPropsLoadMember.h
--- a/Engine/Plugins/Experimental/PlainProps/Source/Public/PlainPropsLoadMember.h
+++ b/Engine/Plugins/Experimental/PlainProps/Source/Public/PlainPropsLoadMember.h
@@ -15,8 +15,8 @@ class FStructRangeLoadIterator;
 
 struct FMemcpyLoadPlan
 {
-	uint32					Size;
-	uint32					Offset;
+	uint32								Size = 0;
+	uint32								Offset = 0;
 };
 
 struct FSchemaLoadHandle
