PlainProps: Add "place breakpoint here" comments to Diff.cpp

#rnx

diff --git a/Engine/Plugins/Experimental/PlainProps/Source/Private/PlainPropsDiff.cpp b/Engine/Plugins/Experimental/PlainProps/Source/Private/PlainPropsDiff.cpp
--- a/Engine/Plugins/Experimental/PlainProps/Source/Private/PlainPropsDiff.cpp
+++ b/Engine/Plugins/Experimental/PlainProps/Source/Private/PlainPropsDiff.cpp
@@ -21,6 +21,7 @@ inline bool DiffMember(FLeafMemberBinding Member, const void* OwnerA, const void
 	const uint8* B = At(OwnerB, Member.Offset);
 	if (DiffLeaf(A, B, Member.Leaf))
 	{
+		// place breakpoint here
 		Ctx.Out.Add({Member.Leaf.Pack(), Name, {Member.Enum}, A, B});
 		return true;
 	}
@@ -63,6 +64,7 @@ bool DiffMember(FRangeMemberBinding Member, const void* OwnerA, const void* Owne
 	{
 		if (bDiff)
 		{
+			// place breakpoint here
 			Ctx.Out.Add({FMemberBindType(Binding.GetSizeType()), Name, {.Range = Binding}, A, B});
 		}
 	}
@@ -93,6 +95,7 @@ static bool DiffMember(FStructMemberBinding Member, const void* OwnerA, const vo
 	const uint8* B = At(OwnerB, Member.Offset);
 	if (DiffStruct(Member.Id, A, B, Ctx))
 	{
+		// place breakpoint here
 		Ctx.Out.Add({FMemberBindType(Member.Type), Name, {.Struct = Member.Id}, A, B});
 		return true;
 	}
