PlainProps: Handle delta serialization of dense (RequireAll) structs by diffing them before saving them out in completeness.
Add unit test.

#rnx

diff --git a/Engine/Plugins/Experimental/PlainProps/Source/Private/PlainPropsSave.cpp b/Engine/Plugins/Experimental/PlainProps/Source/Private/PlainPropsSave.cpp
--- a/Engine/Plugins/Experimental/PlainProps/Source/Private/PlainPropsSave.cpp
+++ b/Engine/Plugins/Experimental/PlainProps/Source/Private/PlainPropsSave.cpp
@@ -516,6 +516,10 @@ FBuiltStruct* SaveStructDeltaIfDiff(const void* Struct, const void* Default, FBi
 	}
 
 	const FSchemaBinding& Schema = Ctx.Schemas.GetStruct(Id, /* out */ Declaration);
+	if (Declaration->Occupancy == EMemberPresence::RequireAll && !DiffStructs(Struct, Default, Id, Ctx))
+	{
+		return nullptr;
+	}
 	SaveSchemaBoundStructDelta(Out, Struct, Default, Schema, *Declaration, Ctx);
 
 	return Out.IsEmpty() ? nullptr : Out.BuildAndReset(Ctx.Scratch, *Declaration, Ctx.Schemas.GetDebug());
diff --git a/Engine/Plugins/Experimental/PlainProps/Source/Tests/PlainPropsUeCoreTests.cpp b/Engine/Plugins/Experimental/PlainProps/Source/Tests/PlainPropsUeCoreTests.cpp
--- a/Engine/Plugins/Experimental/PlainProps/Source/Tests/PlainPropsUeCoreTests.cpp
+++ b/Engine/Plugins/Experimental/PlainProps/Source/Tests/PlainPropsUeCoreTests.cpp
@@ -887,6 +887,19 @@ struct FDelta
 };
 PP_REFLECT_STRUCT(PlainProps::UE::Test, FDelta, void, A, B, C, D, E);
 
+struct FDenseDelta
+{
+	int32 X = 0;
+	int32 Y = 0;
+	friend bool operator==(const FDenseDelta& A, const FDenseDelta& B) = default;
+};
+
+PP_REFLECT_STRUCT(PlainProps::UE::Test, FDenseDelta, void, X, Y);
+
+} namespace PlainProps {
+template<> struct TOccupancyOf<PlainProps::UE::Test::FDenseDelta> : FRequireAll {};
+} namespace PlainProps::UE::Test {
+
 //////////////////////////////////////////////////////////////////////////
 //
 //struct FObject
@@ -1214,12 +1227,15 @@ TEST_CASE_NAMED(FPlainPropsUeCoreTest, "System::Core::Serialization::PlainProps:
 	{
 		TScopedDefaultStructBinding<FInt> Int;
 		TScopedDefaultStructBinding<FDelta> Delta;
+		TScopedDefaultStructBinding<FDenseDelta> DenseDelta;
 		Run([](FBatchSaver& Batch)
 			{
 				FDelta Zero = {false, 0, {}, {}, {}};
 				CHECK(!Batch.SaveDelta(FInt{123},FInt{123}));
 				CHECK(!Batch.SaveDelta(FDelta{},FDelta{}));
 				CHECK(!Batch.SaveDelta(Zero, Zero));
+				CHECK(!Batch.SaveDelta(FDenseDelta{}, FDenseDelta{}));
+				CHECK(!Batch.SaveDelta(FDenseDelta{1, 2}, FDenseDelta{1, 2}));
 					
 				Batch.SaveDelta({}, Zero);
 				Batch.SaveDelta(Zero, {});
@@ -1227,6 +1243,9 @@ TEST_CASE_NAMED(FPlainPropsUeCoreTest, "System::Core::Serialization::PlainProps:
 				Batch.SaveDelta(FDelta{.C = {321}}, {});
 				Batch.SaveDelta(FDelta{.D = {0}}, {});
 				Batch.SaveDelta(FDelta{.E = "!!"}, {});
+				Batch.SaveDelta(FDenseDelta{.X = 1}, {});
+				Batch.SaveDelta(FDenseDelta{.Y = 2}, {});
+				Batch.SaveDelta(FDenseDelta{1, 2}, {});
 			}, 
 			[](FBatchLoader& Batch)
 			{
@@ -1237,6 +1256,9 @@ TEST_CASE_NAMED(FPlainPropsUeCoreTest, "System::Core::Serialization::PlainProps:
 				CHECK(Batch.Load<FDelta>() == FDelta{.C = {321}});
 				CHECK(Batch.Load<FDelta>() == FDelta{.D = {0}});
 				CHECK(Batch.Load<FDelta>() == FDelta{.E = "!!"});
+				CHECK(Batch.Load<FDenseDelta>() == FDenseDelta{.X = 1});
+				CHECK(Batch.Load<FDenseDelta>() == FDenseDelta{.Y = 2});
+				CHECK(Batch.Load<FDenseDelta>() == FDenseDelta{1, 2});
 			});
 	}
 
