PlainProps add mechanism to support archetypes being a super class

Works with custom-bound derived classes

#rnx

diff --git a/Engine/Plugins/Experimental/PlainProps/Source/Private/PlainPropsLoad.cpp b/Engine/Plugins/Experimental/PlainProps/Source/Private/PlainPropsLoad.cpp
--- a/Engine/Plugins/Experimental/PlainProps/Source/Private/PlainPropsLoad.cpp
+++ b/Engine/Plugins/Experimental/PlainProps/Source/Private/PlainPropsLoad.cpp
@@ -6,6 +6,7 @@
 #include "PlainPropsInternalBind.h"
 #include "PlainPropsInternalFormat.h"
 #include "PlainPropsInternalRead.h"
+#include "PlainPropsSave.h" // FBaseline use in dummy FLoadBindingBase::SaveCustom
 #include "Algo/Compare.h"
 #include "Algo/Find.h"
 #include "Containers/Set.h"
@@ -492,7 +493,7 @@ template<class LoadIdBinder>
 
 struct FLoadBindingBase : ICustomBinding
 {
-	virtual void SaveCustom(FMemberBuilder& Dst, const void* Src, const void* Default, const FSaveContext& Ctx) override final { unimplemented(); }
+	virtual void SaveCustom(FMemberBuilder& Dst, const void* Src, FBaseline Base, const FSaveContext& Ctx) override final { unimplemented(); }
 	virtual bool DiffCustom(const void*, const void*, const FBindContext&) const override final { unimplemented(); return false; }
 };
 struct FTypeErasedLoadBinding final : FLoadBindingBase
diff --git a/Engine/Plugins/Experimental/PlainProps/Source/Private/PlainPropsSave.cpp b/Engine/Plugins/Experimental/PlainProps/Source/Private/PlainPropsSave.cpp
--- a/Engine/Plugins/Experimental/PlainProps/Source/Private/PlainPropsSave.cpp
+++ b/Engine/Plugins/Experimental/PlainProps/Source/Private/PlainPropsSave.cpp
@@ -13,6 +13,19 @@
 namespace PlainProps
 {
 
+FBaseline::FBaseline(FDefaultLink* Chain UE_LIFETIMEBOUND, int32 Num)
+: Ptr(reinterpret_cast<uint64>(Chain) | TagBit)
+{
+	check(Num >= 2);
+	for (FDefaultLink& Link : MakeArrayView(Chain, Num))
+	{
+		checkSlow(!(Link.Ptr & TagBit));
+	}
+
+	// Tag last link
+	Chain[Num - 1].Ptr |= TagBit;
+}
+
 static uint64 GetBit(uint8 Byte, uint8 BitIdx)
 {
 	return (Byte >> BitIdx) & 1;
@@ -302,32 +315,9 @@ static void SaveMember(FMemberBuilder& Out, const void* Struct, FMemberId Name,
 	Out.AddStruct(Name, Member.Id, SaveStruct(At(Struct, Member.Offset), Member.Id, Ctx));
 }
 
-// Allow a super that has been prebuilt with SaveStructDelta
-inline void SaveAllMembers(FMemberBuilder& Out, const void* Struct, const FSchemaBinding& Schema, const FStructDeclaration& Declaration, const FSaveContext& Ctx, FBuiltStruct* BuiltSuper = nullptr)
+inline void SaveNamedMembers(FMemberBuilder& Out, const void* Struct, FMemberVisitor It, TConstArrayView<FMemberId> Names, const FSaveContext& Ctx)
 {
-	FMemberVisitor It(Schema);
-	if (Declaration.Super)
-	{
-		FBindId SuperId = It.GrabSuper();
-		checkSlow(SuperId == ToOptionalStruct(Declaration.Super));
-		if (BuiltSuper)
-		{
-			Out.AddSuperStruct(SuperId, BuiltSuper);
-		}
-		else
-		{
-			const FStructDeclaration* SuperDecl;
-			const FSchemaBinding& SuperSchema = Ctx.Schemas.GetStruct(SuperId, /* out */ SuperDecl);
-			SaveAllMembers(Out, Struct, SuperSchema, *SuperDecl, Ctx);
-			Out.BuildSuperStruct(Ctx.Scratch, *SuperDecl, Ctx.Schemas.GetDebug());
-		}
-	}
-	else
-	{
-		checkSlow(!BuiltSuper);
-	}
-
-	for (FMemberId Name : Declaration.GetMemberOrder())
+	for (FMemberId Name : Names)
 	{
 		switch (It.PeekKind())
 		{
@@ -339,10 +329,9 @@ inline void SaveAllMembers(FMemberBuilder& Out, const void* Struct, const FSchem
 	checkSlow(!It.HasMore());
 }
 
-FBuiltStruct* SaveStruct(const void* Struct, FBindId Id, const FSaveContext& Ctx)
+static const FStructDeclaration& SaveStruct(FMemberBuilder& Out, const void* Struct, FBindId Id, const FSaveContext& Ctx)
 {
 	const FStructDeclaration* Declaration = nullptr;
-	FMemberBuilder Out;
 	if (ICustomBinding* Custom = Ctx.Customs.FindStructToSave(Id, /* out */ Declaration))
 	{
 		Custom->SaveCustom(Out, Struct, nullptr, Ctx);
@@ -350,24 +339,25 @@ FBuiltStruct* SaveStruct(const void* Struct, FBindId Id, const FSaveContext& Ctx
 	else
 	{
 		const FSchemaBinding& Schema = Ctx.Schemas.GetStruct(Id, /* out */ Declaration);
-		SaveAllMembers(Out, Struct, Schema, *Declaration, Ctx);
-	}
 
-	return Out.BuildAndReset(Ctx.Scratch, *Declaration, Ctx.Schemas.GetDebug());
-}
+		FMemberVisitor It(Schema);
+		if (Declaration->Super)
+		{
+			const FStructDeclaration& SuperDecl = SaveStruct(Out, Struct, It.GrabSuper(), Ctx);
+			Out.BuildSuperStruct(Ctx.Scratch, SuperDecl, Ctx.Schemas.GetDebug());
+			checkSlow(SuperDecl.Id == ToOptionalStruct(Declaration->Super));
+		}
 
+		SaveNamedMembers(Out, Struct, It, Declaration->GetMemberOrder(), Ctx);
+	}
+	return *Declaration;
+}
 
-FBuiltStruct* SaveStructWithSuper(const void* Struct, FBuiltStruct* BuiltSuper, FBindId Id, const FSaveContext& Ctx)
+FBuiltStruct* SaveStruct(const void* Struct, FBindId Id, const FSaveContext& Ctx)
 {
-	check(BuiltSuper);
-	const FStructDeclaration* Declaration = nullptr;
-	check(!Ctx.Customs.FindStructToSave(Id, /* out */ Declaration));
 	FMemberBuilder Out;
-
-	const FSchemaBinding& Schema = Ctx.Schemas.GetStruct(Id, /* out */ Declaration);
-	check(Declaration->Super);
-	SaveAllMembers(Out, Struct, Schema, *Declaration, Ctx, BuiltSuper);
-	return Out.BuildAndReset(Ctx.Scratch, *Declaration, Ctx.Schemas.GetDebug());
+	const FStructDeclaration& Declaration = SaveStruct(Out, Struct, Id, Ctx);
+	return Out.BuildAndReset(Ctx.Scratch, Declaration, Ctx.Schemas.GetDebug());
 }
 
 ////////////////////////////////////////////////////////////////////////////////////////////////
@@ -446,21 +436,22 @@ static void SaveMemberDelta(FMemberBuilder& Out, const void* Struct, const void*
 	}
 }
 
-static void SaveSchemaBoundStructDelta(FMemberBuilder& Out, const void* Struct, const void* Default, const FSchemaBinding& Schema,  const FStructDeclaration& Declaration, const FSaveContext& Ctx)
+static void SaveSchemaBoundStructDelta(FMemberBuilder& Out, const void* Struct, FBaseline Base, const FSchemaBinding& Schema, const FStructDeclaration& Declaration, const FSaveContext& Ctx)
 {
-	if (Declaration.Occupancy == EMemberPresence::AllowSparse)
+	FMemberVisitor It(Schema);
+	if (Declaration.Super)
 	{
-		FMemberVisitor It(Schema);
-		if (Declaration.Super)
-		{
-			FBindId SuperId = It.GrabSuper();
-			checkSlow(SuperId == ToOptionalStruct(Declaration.Super));
-			const FStructDeclaration* SuperDecl;
-			const FSchemaBinding& SuperSchema = Ctx.Schemas.GetStruct(SuperId, /* out */ SuperDecl);
-			SaveSchemaBoundStructDelta(Out, Struct, Default, SuperSchema, *SuperDecl, Ctx);
-			Out.BuildSuperStruct(Ctx.Scratch, *SuperDecl, Ctx.Schemas.GetDebug());
-		}
+		FBindId SuperId = It.GrabSuper();
+		checkSlow(SuperId == ToOptionalStruct(Declaration.Super));
+		const FStructDeclaration* SuperDecl;
+		const FSchemaBinding& SuperSchema = Ctx.Schemas.GetStruct(SuperId, /* out */ SuperDecl);
+		SaveSchemaBoundStructDelta(Out, Struct, Base.Super(), SuperSchema, *SuperDecl, Ctx);
+		Out.BuildSuperStruct(Ctx.Scratch, *SuperDecl, Ctx.Schemas.GetDebug());
+	}
 
+	const void* Default = Base.Get();
+	if (Declaration.Occupancy == EMemberPresence::AllowSparse && Default)
+	{
 		for (FMemberId Name : Declaration.GetMemberOrder())
 		{
 			switch (It.PeekKind())
@@ -474,54 +465,61 @@ static void SaveSchemaBoundStructDelta(FMemberBuilder& Out, const void* Struct,
 	}
 	else
 	{
-		SaveAllMembers(Out, Struct, Schema, Declaration, Ctx);
+		SaveNamedMembers(Out, Struct, It, Declaration.GetMemberOrder(), Ctx);
 	}
 }
 
-FBuiltStruct* SaveStructDelta(const void* Struct, const void* Default, FBindId Id, const FSaveContext& Ctx)
+const FStructDeclaration& SaveStructDelta(FMemberBuilder& Out, const void* Struct, FBaseline Base, FBindId Id, const FSaveContext& Ctx)
 {
-	FMemberBuilder Out;
 	const FStructDeclaration* Declaration = nullptr;
 	if (ICustomBinding* Custom = Ctx.Customs.FindStructToSave(Id, /* out */ Declaration))
 	{
-		if (Custom->DiffCustom(Struct, Default, Ctx))
+		const void* Default = Base.Get();
+		if (!Default || Custom->DiffCustom(Struct, Default, Ctx))
 		{
-			Custom->SaveCustom(Out, Struct, Default, Ctx);
+			Custom->SaveCustom(Out, Struct, Base, Ctx);
 			// Could type-check saved custom struct against declaration here,
 			// or add a type-checking BuildAndReset overload and call that.
-			// It's not important to type-check bindings, they're generated.
+			// It's not important to type-check generate bindings.
 		}
 	}
 	else
 	{
 		const FSchemaBinding& Schema = Ctx.Schemas.GetStruct(Id, /* out */ Declaration);
-		SaveSchemaBoundStructDelta(Out, Struct, Default, Schema, *Declaration, Ctx);
+		SaveSchemaBoundStructDelta(Out, Struct, Base, Schema, *Declaration, Ctx);
 	}
-	return Out.BuildAndReset(Ctx.Scratch, *Declaration, Ctx.Schemas.GetDebug());
+	
+	return *Declaration;
 }
 
-FBuiltStruct* SaveStructDeltaIfDiff(const void* Struct, const void* Default, FBindId Id, const FSaveContext& Ctx)
+FBuiltStruct* SaveStructDelta(const void* Struct, FBaseline Base, FBindId Id, const FSaveContext& Ctx)
 {
 	FMemberBuilder Out;
-	const FStructDeclaration* Declaration = nullptr;
+	const FStructDeclaration& Declaration = SaveStructDelta(Out, Struct, Base, Id, Ctx);
+	return Out.BuildAndReset(Ctx.Scratch, Declaration, Ctx.Schemas.GetDebug());
+}
+
+FBuiltStruct* SaveStructDeltaIfDiff(const void* Struct, FBaseline Base, FBindId Id, const FSaveContext& Ctx)
+{
+	FMemberBuilder Out;
+	const void* Default = Base.Get();
+	const FStructDeclaration* Declaration;
 	if (ICustomBinding* Custom = Ctx.Customs.FindStructToSave(Id, /* out */ Declaration))
 	{
-		if (Custom->DiffCustom(Struct, Default, Ctx))
+		if (!Default || Custom->DiffCustom(Struct, Default, Ctx))
 		{
-			Custom->SaveCustom(Out, Struct, Default, Ctx);
-			return Out.BuildAndReset(Ctx.Scratch, *Declaration, Ctx.Schemas.GetDebug());
+			Custom->SaveCustom(Out, Struct, Base, Ctx);
 		}
-		
-		return nullptr;
 	}
-
-	const FSchemaBinding& Schema = Ctx.Schemas.GetStruct(Id, /* out */ Declaration);
-	if (Declaration->Occupancy == EMemberPresence::RequireAll && !DiffStructs(Struct, Default, Id, Ctx))
+	else
 	{
-		return nullptr;
+		const FSchemaBinding& Schema = Ctx.Schemas.GetStruct(Id, /* out */ Declaration);
+		if (!Default || Declaration->Occupancy == EMemberPresence::AllowSparse || DiffStructs(Struct, Default, Id, Ctx))
+		{
+			SaveSchemaBoundStructDelta(Out, Struct, Base, Schema, *Declaration, Ctx);
+		}
 	}
-	SaveSchemaBoundStructDelta(Out, Struct, Default, Schema, *Declaration, Ctx);
-
+	
 	return Out.IsEmpty() ? nullptr : Out.BuildAndReset(Ctx.Scratch, *Declaration, Ctx.Schemas.GetDebug());
 }
 
diff --git a/Engine/Plugins/Experimental/PlainProps/Source/Private/PlainPropsWrite.cpp b/Engine/Plugins/Experimental/PlainProps/Source/Private/PlainPropsWrite.cpp
--- a/Engine/Plugins/Experimental/PlainProps/Source/Private/PlainPropsWrite.cpp
+++ b/Engine/Plugins/Experimental/PlainProps/Source/Private/PlainPropsWrite.cpp
@@ -337,7 +337,6 @@ static FOptionalSchemaId GetStaticInnerSchema(const FMemberSchema& Schema, const
 		}
 		else if (!InnermostType.AsStruct().IsDynamic)
 		{
-			//tmptorp
 			return { NewIds.RemapStruct(InnerSchema.Get().AsStructBindId()) };
 		}
 	}
diff --git a/Engine/Plugins/Experimental/PlainProps/Source/Public/PlainPropsBind.h b/Engine/Plugins/Experimental/PlainProps/Source/Public/PlainPropsBind.h
--- a/Engine/Plugins/Experimental/PlainProps/Source/Public/PlainPropsBind.h
+++ b/Engine/Plugins/Experimental/PlainProps/Source/Public/PlainPropsBind.h
@@ -18,6 +18,7 @@
 namespace PlainProps 
 {
 
+class FBaseline;
 struct FBindContext;
 struct FBuiltRange;
 struct FBuiltStruct;
@@ -264,7 +265,7 @@ enum ECustomLoadMethod { Construct, Assign };
 struct ICustomBinding
 {
 	virtual ~ICustomBinding() {}
-	virtual void	SaveCustom(FMemberBuilder& Dst, const void* Src, const void* Default, const FSaveContext& Ctx) = 0;
+	virtual void	SaveCustom(FMemberBuilder& Dst, const void* Src, FBaseline Base, const FSaveContext& Ctx) = 0;
 	virtual void	LoadCustom(void* Dst, FStructLoadView Src, ECustomLoadMethod Method) const = 0;
 	virtual bool	DiffCustom(const void* StructA, const void* StructB, const FBindContext& Ctx) const = 0;
 
@@ -914,9 +915,13 @@ struct TCustomSpecifier : TCustomInit<Ids>
 	}
 };
 
-// Workaround for clang requiring FStructLoadView definition, make it dependent on any other type
+// Workaround for clang requiring the definition, make it dependent on any other type
 template<typename>
-struct TDependent { using StructLoadView = FStructLoadView; };
+struct TDependent
+{
+	using Baseline = FBaseline;
+	using StructLoadView = FStructLoadView;
+};
 
 // Scoped custom binding for native type
 template<class T, class Runtime, typename CustomBinding = CustomBind<T>>
@@ -952,8 +957,9 @@ struct TScopedStructBinding : FBothStructId, CustomBinding
 		Runtime::GetCustoms().DropStruct(BindId);
 	}
 
-	virtual void SaveCustom(FMemberBuilder& Dst, const void* Src, const void* Default, const FSaveContext& Ctx) override
+	virtual void SaveCustom(FMemberBuilder& Dst, const void* Src, TDependent<Ids>::Baseline Base, const FSaveContext& Ctx) override
 	{
+		const void* Default = Base.Get();
 		CustomBinding::Save(Dst, *static_cast<const T*>(Src), static_cast<const T*>(Default), Ctx);
 	}
 
diff --git a/Engine/Plugins/Experimental/PlainProps/Source/Public/PlainPropsSave.h b/Engine/Plugins/Experimental/PlainProps/Source/Public/PlainPropsSave.h
--- a/Engine/Plugins/Experimental/PlainProps/Source/Public/PlainPropsSave.h
+++ b/Engine/Plugins/Experimental/PlainProps/Source/Public/PlainPropsSave.h
@@ -23,12 +23,63 @@ FSaveContext MakeSaveContext(FScratchAllocator& Scratch)
 
 ////////////////////////////////////////////////////////////////////////////////////////////////
 
+// FBaseline helper
+struct FDefaultLink
+{
+	FDefaultLink(const void* Default) : Ptr(reinterpret_cast<uint64>(Default)) {}
+	uint64 Ptr;
+};
+
+// A default pointer or a short chain of super default pointers,
+// where last default is shared by all remaining supers
+//
+// Optimized for normal case where default is the most derived type being saved.
+// Handles default being some super class, via nullptr(s) first and default last.
+// Can also help multiple inheritance impacting super offsets. 
+class FBaseline
+{
+public:
+	FBaseline(const void* Default) : Ptr(reinterpret_cast<uint64>(Default)) {}
+	PLAINPROPS_API FBaseline(FDefaultLink* Chain UE_LIFETIMEBOUND, int32 Num);
+
+	const void*				Get() const;	// Get default instance for current type
+	FBaseline				Super() const;	// Get baseline for super type
+
+private:
+	static constexpr uint64 TagBit = uint64(1) << FPlatformMemory::KernelAddressBit;
+	uint64 Ptr = 0;
+	FBaseline(uint64 In) : Ptr(In) {}
+};
+
+////////////////////////////////////////////////////////////////////////////////////////////////
+
 [[nodiscard]] PLAINPROPS_API FBuiltStruct*	SaveStruct(const void* Struct, FBindId BindId, const FSaveContext& Context);
-// todo: temp api for delta serialization with a superclass default object
-[[nodiscard]] PLAINPROPS_API FBuiltStruct*	SaveStructWithSuper(const void* Struct, FBuiltStruct* BuiltSuper, FBindId BindId, const FSaveContext& Context);
-[[nodiscard]] PLAINPROPS_API FBuiltStruct*	SaveStructDelta(const void* Struct, const void* Default, FBindId BindId, const FSaveContext& Context);
-[[nodiscard]] PLAINPROPS_API FBuiltStruct*	SaveStructDeltaIfDiff(const void* Struct, const void* Default, FBindId BindId, const FSaveContext& Context);
+[[nodiscard]] PLAINPROPS_API FBuiltStruct*	SaveStructDelta(const void* Struct, FBaseline Base, FBindId BindId, const FSaveContext& Context);
+[[nodiscard]] PLAINPROPS_API FBuiltStruct*	SaveStructDeltaIfDiff(const void* Struct, FBaseline Base, FBindId BindId, const FSaveContext& Context);
 [[nodiscard]] PLAINPROPS_API FBuiltRange*	SaveRange(const void* Range, FRangeMemberBinding Member, const FSaveContext& Ctx);
 [[nodiscard]] PLAINPROPS_API FBuiltRange*	SaveLeafRange(const void* Range, const ILeafRangeBinding& Binding, FUnpackedLeafType Leaf, const FSaveContext& Ctx);
 
+////////////////////////////////////////////////////////////////////////////////////////////////
+
+inline const void* FBaseline::Get() const
+{
+	// AsLink()->Ptr is untagged, only last link is tagged to mark end and GetSuper() resolves it 
+    uint64 Out = (Ptr & TagBit) ? reinterpret_cast<const FDefaultLink*>(Ptr & ~TagBit)->Ptr : Ptr;
+    return reinterpret_cast<const void*>(Out);
+}
+
+inline FBaseline FBaseline::Super() const
+{
+    if (Ptr & TagBit)
+    {
+		// Return next link in chain, but resolve the last tagged link to an "inline" archetype
+		FDefaultLink SuperLink = reinterpret_cast<const FDefaultLink*>(Ptr & ~TagBit)[1];
+		return (SuperLink.Ptr & TagBit) ? FBaseline(SuperLink.Ptr & ~TagBit)
+										: FBaseline(Ptr + sizeof(FDefaultLink));
+	}
+
+	// Normal case, reuse the same archetype for all super classes
+	return *this;
+}
+
 } // namespace PlainProps
diff --git a/Engine/Plugins/Experimental/PlainProps/Source/Tests/PlainPropsUeCoreTests.cpp b/Engine/Plugins/Experimental/PlainProps/Source/Tests/PlainPropsUeCoreTests.cpp
--- a/Engine/Plugins/Experimental/PlainProps/Source/Tests/PlainPropsUeCoreTests.cpp
+++ b/Engine/Plugins/Experimental/PlainProps/Source/Tests/PlainPropsUeCoreTests.cpp
@@ -109,7 +109,7 @@ struct FNameBinding final : ICustomBinding
 		return PlainProps::Declare({Id, NoId, 0, EMemberPresence::RequireAll, {IdxName}, {Specify<int32>()}});
 	}
 
-	virtual void SaveCustom(FMemberBuilder& Dst, const void* Src, const void*, const FSaveContext& Ctx) override
+	virtual void SaveCustom(FMemberBuilder& Dst, const void* Src, FBaseline, const FSaveContext& Ctx) override
 	{
 		FSetElementId Idx = Names.Add(*static_cast<const FName*>(Src));
 		Dst.Add(IdxName, Idx.AsInteger());
diff --git a/Engine/Plugins/Experimental/PlainPropsUObject/Source/Private/PlainPropsRoundtripTest.cpp b/Engine/Plugins/Experimental/PlainPropsUObject/Source/Private/PlainPropsRoundtripTest.cpp
--- a/Engine/Plugins/Experimental/PlainPropsUObject/Source/Private/PlainPropsRoundtripTest.cpp
+++ b/Engine/Plugins/Experimental/PlainPropsUObject/Source/Private/PlainPropsRoundtripTest.cpp
@@ -154,8 +154,9 @@ struct TCustomPropertyBinding final : ICustomBinding
 
 	BatchType& Batch;
 
-	virtual void SaveCustom(FMemberBuilder& Dst, const void* Src, const void* Default, const FSaveContext& Ctx) override
+	virtual void SaveCustom(FMemberBuilder& Dst, const void* Src, FBaseline Base, const FSaveContext& Ctx) override
 	{
+		const void* Default = Base.Get();
 		if (!Default || DiffCustom(Src, Default, Ctx))
 		{
 			Batch.Save(Dst, *static_cast<const Type*>(Src), Ctx);
@@ -247,22 +248,37 @@ public:
 		ObjectMap.Reserve(NumReserve);
 	}
 
-	// If Arch, delta serialize the matching part of the object inheritance chain
+	// If Arch, delta-serialize the matching part of the object inheritance chain
 	FBuiltStruct* SaveImpl(FBindId Id, const UStruct* Class, const UObject* Object, const UStruct* ArchClass, const UObject* Arch)
 	{
 		if (ArchClass == nullptr)
 		{
 			return SaveStruct(Object, Id, FlatCtx);
 		}
-		if (Class == ArchClass)
+		else if (Class == ArchClass)
 		{
-			return SaveStructDelta(Object, Arch, Id, DeltaCtx);
+			return SaveStructDelta(Object, FBaseline(Arch), Id, DeltaCtx);
 		}
 
-		Class = SkipEmptyBases(Class->GetInheritanceSuper());
-		FBindId SuperId = GUE.Names.IndexBindId(IndexType(Class));
-		FBuiltStruct* Super = SaveImpl(SuperId, Class, Object, ArchClass, Arch);
-		return SaveStructWithSuper(Object, Super, Id, FlatCtx);
+		// Archetype is some super class, save derived part(s) with null archetype
+		check(Class->IsChildOf(ArchClass));
+		TArray<FDefaultLink, TInlineAllocator<4>> Chain;
+		const UStruct* SuperIt = Class;
+		while ((SuperIt = SkipEmptyBases(SuperIt->GetInheritanceSuper())))
+		{
+			Chain.Emplace(nullptr);
+			if (SuperIt == ArchClass)
+			{
+				Chain.Emplace(Arch);
+				return SaveStructDelta(Object, FBaseline(Chain.GetData(), Chain.Num()), Id, DeltaCtx);
+			}
+		}
+
+		checkf(!SkipEmptyBases(ArchClass), TEXT("Code missing for skipped ArchClass without skipping all ArchClass's supers"));
+
+		// ArchClass was skipped so skip delta serialization, but use
+		// DeltaCtx for IDefaultStructs lookup of range-of-structs
+		return SaveStruct(Object, Id, DeltaCtx);
 	}
 
 	void Save(FBindId Id, const UObject* Object, const UObject* Arch)
@@ -675,6 +691,12 @@ struct FInstance
 		Class->GetDefaultObject(/* create lazily */ true);
 		Arch = Orig->GetArchetype();
 		check(Arch != Orig);
+
+		if (Arch->GetClass() != Class)
+		{
+			Arch->GetClass()->GetDefaultObject(/* create lazily */ true);
+		}
+
 		Id = GUE.Names.IndexBindId(IndexType(Class));
 	}
 };
@@ -687,7 +709,7 @@ static UObject* MakeEmptyInstance(UObject* Obj, FName Name)
 	Params.SetFlags = Obj->GetFlags();
 	Params.Template = Obj->GetArchetype();
 	Params.bAssumeTemplateIsArchetype = true;
-	Params.bCopyTransientsFromClassDefaults = true;
+	Params.bCopyTransientsFromClassDefaults = !(Params.SetFlags & RF_ClassDefaultObject);
 	return StaticConstructObject_Internal(Params);
 }
 
@@ -730,8 +752,8 @@ static bool IncludeClass(UClass* Class)
 FORCENOINLINE static void SavePlainProps(FBatchSaver& Batch, TConstArrayView<FInstance> Instances)
 {
 	for (const FInstance& Instance : Instances)
-	{	
-		Batch.Save(Instance.Id, Instance.Orig, Instance.Base);
+	{
+		Batch.Save(Instance.Id, Instance.Orig, Instance.Arch);
 	}
 }
 
@@ -1330,6 +1352,11 @@ TMap<FName, TUniquePtr<FPlainPropsLinkerSandbox::FLinkerBatch>> FPlainPropsLinke
 
 ////////////////////////////////////////////////////////////////////////////////
 
+static bool IsNativeCDO(UObject* Object)
+{
+	return Object->HasAllFlags(RF_ClassDefaultObject) && Object->GetClass()->HasAllClassFlags(CLASS_Native);
+}
+
 int32 RoundtripViaBatch(TConstArrayView<UObject*> Objects, ERoundtrip Options)
 {
 	TArray<FInstance> Instances;
@@ -1352,12 +1379,12 @@ int32 RoundtripViaBatch(TConstArrayView<UObject*> Objects, ERoundtrip Options)
 	else
 	{
 		UE_LOGFMT(LogPlainPropsUObject, Display, "Gathering all non-empty UObjects...");
-		static constexpr EObjectFlags SkipFlags = RF_ClassDefaultObject | RF_MirroredGarbage | RF_InheritableComponentTemplate;
+		static constexpr EObjectFlags SkipFlags =  RF_MirroredGarbage | RF_InheritableComponentTemplate;
 		for (TObjectIterator<UObject> It(SkipFlags); It; ++It)
 		{
 			UObject* Object = *It;
 			UClass* Class = Object->GetClass();
-			if (ShouldBind(Class) && IncludeClass(Class))
+			if (ShouldBind(Class) && IncludeClass(Class) && !IsNativeCDO(Object))
 			{
 				Instances.Emplace(Object->GetPathName(), Object);
 			}
diff --git a/Engine/Plugins/Experimental/PlainPropsUObject/Source/Private/PlainPropsUObjectBindings.cpp b/Engine/Plugins/Experimental/PlainPropsUObject/Source/Private/PlainPropsUObjectBindings.cpp
--- a/Engine/Plugins/Experimental/PlainPropsUObject/Source/Private/PlainPropsUObjectBindings.cpp
+++ b/Engine/Plugins/Experimental/PlainPropsUObject/Source/Private/PlainPropsUObjectBindings.cpp
@@ -629,8 +629,9 @@ struct FMulticastSparseDelegateBinding final : ICustomBinding
 	const FName OwningClassName;
 	const FName DelegateName;
 
-	virtual void SaveCustom(FMemberBuilder& Dst, const void* Src, const void* Default, const FSaveContext& Ctx) override
+	virtual void SaveCustom(FMemberBuilder& Dst, const void* Src, FBaseline Base, const FSaveContext& Ctx) override
 	{
+		const void* Default = Base.Get();
 		if (!Default || DiffCustom(Src, Default, Ctx))
 		{
 			Save(Dst, *static_cast<const FSparseDelegate*>(Src), Ctx);
@@ -799,8 +800,9 @@ struct FInterfaceBinding final : ICustomBinding
 
 	const TObjectPtr<UClass> InterfaceClass;
 
-	virtual void SaveCustom(FMemberBuilder& Dst, const void* Src, const void* Default, const FSaveContext& Ctx) override
+	virtual void SaveCustom(FMemberBuilder& Dst, const void* Src, FBaseline Base, const FSaveContext& Ctx) override
 	{
+		const void* Default = Base.Get();
 		if (!Default || DiffCustom(Src, Default, Ctx))
 		{
 			Save(Dst, *static_cast<const FScriptInterface*>(Src), Ctx);
@@ -2047,9 +2049,9 @@ struct TSetPropertyBinding : IItemRangeBinding, ICustomBinding
 		}
 	}
 
-	virtual void SaveCustom(FMemberBuilder& Dst, const void* Src, const void* Default, const FSaveContext& Ctx) override
+	virtual void SaveCustom(FMemberBuilder& Dst, const void* Src, FBaseline Base, const FSaveContext& Ctx) override
 	{
-		SaveSetDelta(*this, Dst, *static_cast<const FScriptSet*>(Src), static_cast<const FScriptSet*>(Default), Ctx);
+		SaveSetDelta(*this, Dst, *static_cast<const FScriptSet*>(Src), static_cast<const FScriptSet*>(Base.Get()), Ctx);
 	}
 
 	virtual void LoadCustom(void* Dst, FStructLoadView Src, ECustomLoadMethod Method) const override
@@ -2465,9 +2467,9 @@ struct TMapPropertyCustomBinding : TMapPropertyItemBinding<FScriptMap, KeyKind,
 		}
 	}
 
-	virtual void SaveCustom(FMemberBuilder& Dst, const void* Src, const void* Default, const FSaveContext& Ctx) override
+	virtual void SaveCustom(FMemberBuilder& Dst, const void* Src, FBaseline Base, const FSaveContext& Ctx) override
 	{
-		SaveSetDelta(*this, Dst, *static_cast<const FScriptMap*>(Src), static_cast<const FScriptMap*>(Default), Ctx);
+		SaveSetDelta(*this, Dst, *static_cast<const FScriptMap*>(Src), static_cast<const FScriptMap*>(Base.Get()), Ctx);
 	}
 
 	virtual void LoadCustom(void* Dst, FStructLoadView Src, ECustomLoadMethod Method) const override
