PlainProps add SaveSuperDelta and TryLoadSuper helpers

#rnx

diff --git a/Engine/Plugins/Experimental/PlainProps/Source/Private/PlainPropsRead.cpp b/Engine/Plugins/Experimental/PlainProps/Source/Private/PlainPropsRead.cpp
--- a/Engine/Plugins/Experimental/PlainProps/Source/Private/PlainPropsRead.cpp
+++ b/Engine/Plugins/Experimental/PlainProps/Source/Private/PlainPropsRead.cpp
@@ -183,14 +183,14 @@ void FNestedRangeIterator::operator++()
 
 const FMemberType*	FMemberReader::GetMemberTypes() const	{ return FStructSchema::GetMemberTypes(Footer); }
 const FMemberType*	FMemberReader::GetRangeTypes() const	{ return FStructSchema::GetRangeTypes(Footer, NumMembers); }
-const FSchemaId*	FMemberReader::GetInnerSchemas() const	{ return FStructSchema::GetInnerSchemas(Footer, NumMembers, NumRangeTypes, NumMembers - HasSuper); }
+const FSchemaId*	FMemberReader::GetInnerSchemas() const	{ return FStructSchema::GetInnerSchemas(Footer, NumMembers, NumRangeTypes, NumMembers - UsesSuper); }
 const FMemberId*	FMemberReader::GetMemberNames() const	{ return FStructSchema::GetMemberNames(Footer, NumMembers, NumRangeTypes); }
 
 FMemberReader::FMemberReader(const FStructSchema& Schema, FByteReader Values, FSchemaBatchId InBatch)
 : Footer(Schema.Footer)
 , Batch(InBatch)
 , IsSparse(!Schema.IsDense)
-, HasSuper(UsesSuper(Schema.Inheritance))
+, UsesSuper(PlainProps::UsesSuper(Schema.Inheritance))
 , NumMembers(Schema.NumMembers)
 , NumRangeTypes(Schema.NumRangeTypes)
 , InnerSchemaIdx(SkipDeclaredSuperSchema(Schema.Inheritance))
@@ -212,13 +212,13 @@ FMemberReader::FMemberReader(const FStructSchema& Schema, FByteReader Values, FS
 
 FOptionalMemberId FMemberReader::PeekName() const
 {
-	int32 MemberNameIdx = MemberIdx - HasSuper;
+	int32 MemberNameIdx = MemberIdx - UsesSuper;
 	return MemberNameIdx >= 0 ? ToOptional(GetMemberNames()[MemberNameIdx]) : NoId;
 }
 
 FOptionalMemberId FMemberReader::PeekNameUnchecked() const
 {
-	int32 MemberNameIdx = MemberIdx - HasSuper;
+	int32 MemberNameIdx = MemberIdx - UsesSuper;
 	return GetMemberNames()[MemberNameIdx];
 }
 
diff --git a/Engine/Plugins/Experimental/PlainProps/Source/Private/PlainPropsSave.cpp b/Engine/Plugins/Experimental/PlainProps/Source/Private/PlainPropsSave.cpp
--- a/Engine/Plugins/Experimental/PlainProps/Source/Private/PlainPropsSave.cpp
+++ b/Engine/Plugins/Experimental/PlainProps/Source/Private/PlainPropsSave.cpp
@@ -523,6 +523,12 @@ FBuiltStruct* SaveStructDeltaIfDiff(const void* Struct, FBaseline Base, FBindId
 	return Out.IsEmpty() ? nullptr : Out.BuildAndReset(Ctx.Scratch, *Declaration, Ctx.Schemas.GetDebug());
 }
 
+void SaveSuperDelta(FMemberBuilder& Out, const void* Super, FBaseline SuperBase, FBindId SuperId, const FSaveContext& Ctx)
+{
+	const FStructDeclaration& Declaration = SaveStructDelta(Out, Super, SuperBase, SuperId, Ctx);
+	Out.BuildSuperStruct(Ctx.Scratch, Declaration, Ctx.Schemas.GetDebug());
+}
+
 ////////////////////////////////////////////////////////////////////////////////////////////////
 
 FRangeSaverBase::FRangeSaverBase(FScratchAllocator& Scratch, uint64 Num, SIZE_T ItemSize)
diff --git a/Engine/Plugins/Experimental/PlainProps/Source/Public/PlainPropsLoadMember.h b/Engine/Plugins/Experimental/PlainProps/Source/Public/PlainPropsLoadMember.h
--- a/Engine/Plugins/Experimental/PlainProps/Source/Public/PlainPropsLoadMember.h
+++ b/Engine/Plugins/Experimental/PlainProps/Source/Public/PlainPropsLoadMember.h
@@ -205,6 +205,7 @@ public:
 	PLAINPROPS_API explicit FMemberLoader(FStructLoadView Struct);
 
 	bool							HasMore() const				{ return Reader.HasMore(); }
+	bool							HasSuper() const			{ return Reader.HasSuper(); }
 	
 	FOptionalMemberId				PeekName() const			{ return Reader.PeekName(); }
 	FOptionalMemberId				PeekNameUnchecked() const	{ return Reader.PeekNameUnchecked(); }
@@ -214,13 +215,21 @@ public:
 	FLeafView						GrabLeaf()					{ return Reader.GrabLeaf(); }
 	PLAINPROPS_API FRangeLoadView	GrabRange();				// @pre PeekKind() == EMemberKind::Range
 	PLAINPROPS_API FStructLoadView	GrabStruct();				// @pre PeekKind() == EMemberKind::Struct
-
 private:
 	FMemberReader			Reader;
 	const FStructSchemaId*	LoadIdIt;
 	const FLoadBatch&		Batch;
 };
 
+// Grab and load super if it exists, call before grabbing from Src
+inline void	TryLoadSuper(void* Dst, /* in-out */ FMemberLoader& Src)
+{
+	if (Src.HasSuper())
+	{
+		LoadStruct(Dst, Src.GrabStruct());
+	}
+}
+
 //////////////////////////////////////////////////////////////////////////
 
 template<LeafType T>
@@ -231,7 +240,7 @@ void LoadSole(void* Dst, FStructLoadView Src)
 	FMemory::Memcpy(Dst, Src.Values.Peek(), sizeof(T));
 }
 
-template<>
+template<> inline
 void LoadSole<bool>(void* Dst, FStructLoadView Src)
 {
 	Src.Values.CheckSize(sizeof(bool));
diff --git a/Engine/Plugins/Experimental/PlainProps/Source/Public/PlainPropsRead.h b/Engine/Plugins/Experimental/PlainProps/Source/Public/PlainPropsRead.h
--- a/Engine/Plugins/Experimental/PlainProps/Source/Public/PlainPropsRead.h
+++ b/Engine/Plugins/Experimental/PlainProps/Source/Public/PlainPropsRead.h
@@ -620,9 +620,10 @@ public:
 	FMemberReader(const FStructSchema& Schema, FByteReader Values, FSchemaBatchId InBatch);
 
 	bool								HasMore() const			{ return MemberIdx < NumMembers; }
+	bool								HasSuper() const		{ return UsesSuper && MemberIdx == 0; }		
 	
 	PLAINPROPS_API FOptionalMemberId	PeekName() const;			// @pre HasMore()
-	PLAINPROPS_API FOptionalMemberId	PeekNameUnchecked() const; // @pre HasMore()
+	PLAINPROPS_API FOptionalMemberId	PeekNameUnchecked() const;  // @pre HasMore()
 	PLAINPROPS_API EMemberKind			PeekKind() const;			// @pre HasMore()
 	PLAINPROPS_API FMemberType			PeekType() const;			// @pre HasMore()
 
@@ -631,6 +632,7 @@ public:
 	PLAINPROPS_API FStructView			GrabStruct();				// @pre PeekKind() == EMemberKind::Struct
 	//FAnyMemberView					GrabAny();					// @pre HasMore()
 
+
 	// Experimental 
 	// @pre Has N more contiguous members of the expected leaf type
 	template<Arithmetic T>
@@ -645,7 +647,7 @@ protected: // for unit tests
 	const FMemberType*		Footer;
 	const FSchemaBatchId	Batch;					// Needed to resolve schemas
 	const bool				IsSparse : 1;
-	const bool				HasSuper : 1;
+	const bool				UsesSuper : 1;
 	const uint32			NumMembers;
 	const uint32			NumRangeTypes;			// Number of ranges and nested ranges
 
diff --git a/Engine/Plugins/Experimental/PlainProps/Source/Public/PlainPropsSave.h b/Engine/Plugins/Experimental/PlainProps/Source/Public/PlainPropsSave.h
--- a/Engine/Plugins/Experimental/PlainProps/Source/Public/PlainPropsSave.h
+++ b/Engine/Plugins/Experimental/PlainProps/Source/Public/PlainPropsSave.h
@@ -58,6 +58,10 @@ private:
 [[nodiscard]] PLAINPROPS_API FBuiltStruct*	SaveStructDeltaIfDiff(const void* Struct, FBaseline Base, FBindId BindId, const FSaveContext& Context);
 [[nodiscard]] PLAINPROPS_API FBuiltRange*	SaveRange(const void* Range, FRangeMemberBinding Member, const FSaveContext& Ctx);
 [[nodiscard]] PLAINPROPS_API FBuiltRange*	SaveLeafRange(const void* Range, const ILeafRangeBinding& Binding, FUnpackedLeafType Leaf, const FSaveContext& Ctx);
+[[nodiscard]] PLAINPROPS_API FBuiltRange*	SaveStructRange(TConstArrayView<FBuiltStruct*> Structs, const FSaveContext& Ctx);
+
+// Save Super if any members differ from SuperBase
+PLAINPROPS_API void							SaveSuperDelta(FMemberBuilder& Out, const void* Super, FBaseline SuperBase, FBindId SuperId, const FSaveContext& Ctx);
 
 ////////////////////////////////////////////////////////////////////////////////////////////////
 
