PlainProps improve PeekNameUnchecked() pre-condition and change return type

#rnx

diff --git a/Engine/Plugins/Experimental/PlainProps/Source/Private/PlainPropsRead.cpp b/Engine/Plugins/Experimental/PlainProps/Source/Private/PlainPropsRead.cpp
--- a/Engine/Plugins/Experimental/PlainProps/Source/Private/PlainPropsRead.cpp
+++ b/Engine/Plugins/Experimental/PlainProps/Source/Private/PlainPropsRead.cpp
@@ -216,8 +216,9 @@ FOptionalMemberId FMemberReader::PeekName() const
 	return MemberNameIdx >= 0 ? ToOptional(GetMemberNames()[MemberNameIdx]) : NoId;
 }
 
-FOptionalMemberId FMemberReader::PeekNameUnchecked() const
+FMemberId FMemberReader::PeekNameUnchecked() const
 {
+	checkSlow(!UsesSuper);
 	int32 MemberNameIdx = MemberIdx - UsesSuper;
 	return GetMemberNames()[MemberNameIdx];
 }
diff --git a/Engine/Plugins/Experimental/PlainProps/Source/Public/PlainPropsLoadMember.h b/Engine/Plugins/Experimental/PlainProps/Source/Public/PlainPropsLoadMember.h
--- a/Engine/Plugins/Experimental/PlainProps/Source/Public/PlainPropsLoadMember.h
+++ b/Engine/Plugins/Experimental/PlainProps/Source/Public/PlainPropsLoadMember.h
@@ -208,7 +208,7 @@ public:
 	bool							HasSuper() const			{ return Reader.HasSuper(); }
 	
 	FOptionalMemberId				PeekName() const			{ return Reader.PeekName(); }
-	FOptionalMemberId				PeekNameUnchecked() const	{ return Reader.PeekNameUnchecked(); }
+	FMemberId						PeekNameUnchecked() const	{ return Reader.PeekNameUnchecked(); }
 	EMemberKind						PeekKind() const			{ return Reader.PeekKind(); }
 	FMemberType						PeekType() const			{ return Reader.PeekType(); }
 
diff --git a/Engine/Plugins/Experimental/PlainProps/Source/Public/PlainPropsRead.h b/Engine/Plugins/Experimental/PlainProps/Source/Public/PlainPropsRead.h
--- a/Engine/Plugins/Experimental/PlainProps/Source/Public/PlainPropsRead.h
+++ b/Engine/Plugins/Experimental/PlainProps/Source/Public/PlainPropsRead.h
@@ -623,7 +623,7 @@ public:
 	bool								HasSuper() const		{ return UsesSuper && MemberIdx == 0; }		
 	
 	PLAINPROPS_API FOptionalMemberId	PeekName() const;			// @pre HasMore()
-	PLAINPROPS_API FOptionalMemberId	PeekNameUnchecked() const;  // @pre HasMore()
+	PLAINPROPS_API FMemberId			PeekNameUnchecked() const;  // @pre HasMore() && !HasSuper()
 	PLAINPROPS_API EMemberKind			PeekKind() const;			// @pre HasMore()
 	PLAINPROPS_API FMemberType			PeekType() const;			// @pre HasMore()
 
