PlainProps improve misc stuff reported by BugHawk

#rnx

diff --git a/Engine/Plugins/Experimental/PlainProps/Source/Private/PlainPropsBuild.cpp b/Engine/Plugins/Experimental/PlainProps/Source/Private/PlainPropsBuild.cpp
--- a/Engine/Plugins/Experimental/PlainProps/Source/Private/PlainPropsBuild.cpp
+++ b/Engine/Plugins/Experimental/PlainProps/Source/Private/PlainPropsBuild.cpp
@@ -13,9 +13,12 @@ namespace PlainProps
 
 uint8* FScratchAllocator::AllocateInNewPage(SIZE_T Size, uint32 Alignment)
 {
-	if (Size >= DataSize || (DataSize - Size) < static_cast<uint32>(PageEnd - Cursor))
+	const uint32 Padding = Alignment > MIN_ALIGNMENT ? Alignment - MIN_ALIGNMENT : 0u;
+
+	if (Size + Padding >= DataSize || (DataSize - Size - Padding) < static_cast<uint32>(PageEnd - Cursor))
 	{
-		SIZE_T LonePageSize = Align(offsetof(FPage, Data) + Size, Alignment);
+		SIZE_T LonePageSize = offsetof(FPage, Data) + Size + Padding;
+
 		FPage*& PrevPage = LastPage ? LastPage->PrevPage : LastPage;
 		FPage Header = { PrevPage };
 		PrevPage = new (FMemory::Malloc(LonePageSize)) FPage { Header };
diff --git a/Engine/Plugins/Experimental/PlainProps/Source/Private/PlainPropsBuildSchema.cpp b/Engine/Plugins/Experimental/PlainProps/Source/Private/PlainPropsBuildSchema.cpp
--- a/Engine/Plugins/Experimental/PlainProps/Source/Private/PlainPropsBuildSchema.cpp
+++ b/Engine/Plugins/Experimental/PlainProps/Source/Private/PlainPropsBuildSchema.cpp
@@ -454,7 +454,7 @@ void FEnumSchemaBuilder::NoteValue(ELeafWidth Width, uint64 Value, FStructId Str
 			for (int32 Idx = NumValidated, Num = NotedConstants.Num(); Idx < Num; ++Idx)
 			{
 				uint64 Flag = NotedConstants.Get(FSetElementId::FromInteger(Idx));
-				checkf(Algo::FindBy(Declaration.GetEnumerators(), Flag, &FEnumerator::Constant), TEXT("Enum flag %d is undeclared in %s, illegal value detected in %s::%s"), Flag, *Debug.Print(Id), *Debug.Print(Struct), *Debug.Print(Member));
+				checkf(Algo::FindBy(Declaration.GetEnumerators(), Flag, &FEnumerator::Constant), TEXT("Enum flag %llx is undeclared in %s, illegal value detected in %s::%s"), Flag, *Debug.Print(Id), *Debug.Print(Struct), *Debug.Print(Member));
 			}
 		}
 	}
@@ -464,7 +464,7 @@ void FEnumSchemaBuilder::NoteValue(ELeafWidth Width, uint64 Value, FStructId Str
 		NotedConstants.FindOrAdd(Value, /* out */ &bValidated);
 		if (!bValidated)
 		{
-			checkf(Algo::FindBy(Declaration.GetEnumerators(), Value, &FEnumerator::Constant), TEXT("Enum value %d is undeclared in %s, illegal value detected in %s::%s"), Value, *Debug.Print(Id), *Debug.Print(Struct), *Debug.Print(Member));
+			checkf(Algo::FindBy(Declaration.GetEnumerators(), Value, &FEnumerator::Constant), TEXT("Enum value %llu is undeclared in %s, illegal value detected in %s::%s"), Value, *Debug.Print(Id), *Debug.Print(Struct), *Debug.Print(Member));
 		}
 	}
 }
diff --git a/Engine/Plugins/Experimental/PlainProps/Source/Private/PlainPropsDiff.cpp b/Engine/Plugins/Experimental/PlainProps/Source/Private/PlainPropsDiff.cpp
--- a/Engine/Plugins/Experimental/PlainProps/Source/Private/PlainPropsDiff.cpp
+++ b/Engine/Plugins/Experimental/PlainProps/Source/Private/PlainPropsDiff.cpp
@@ -362,7 +362,7 @@ static bool DiffLeaf(FLeafView A, FLeafView B)
 
 static bool DiffLeaves(FUnpackedLeafType Leaf, FLeafRangeView A, FLeafRangeView B, uint64& OutIdx)
 {
-	OutIdx = ~0u;
+	OutIdx = ~uint64(0);
 	switch (Leaf.Type)
 	{
 	case ELeafType::Bool:
diff --git a/Engine/Plugins/Experimental/PlainProps/Source/Private/PlainPropsIndex.cpp b/Engine/Plugins/Experimental/PlainProps/Source/Private/PlainPropsIndex.cpp
--- a/Engine/Plugins/Experimental/PlainProps/Source/Private/PlainPropsIndex.cpp
+++ b/Engine/Plugins/Experimental/PlainProps/Source/Private/PlainPropsIndex.cpp
@@ -73,6 +73,8 @@ FParametricTypeId FParametricTypeIndexer::Index(FParametricTypeView View)
 {
 	using namespace ParametricTypeHash;
 
+	checkSlow(View.NumParameters);
+
 	static constexpr uint32 MinSlack = 4;
 	uint32 WantedSlots = FMath::RoundUpToPowerOfTwo((Types.Num() + MinSlack) * /* 90% load factor */ 10 / 9);
 	if (WantedSlots > NumSlots)
diff --git a/Engine/Plugins/Experimental/PlainProps/Source/Private/PlainPropsInternalBind.h b/Engine/Plugins/Experimental/PlainProps/Source/Private/PlainPropsInternalBind.h
--- a/Engine/Plugins/Experimental/PlainProps/Source/Private/PlainPropsInternalBind.h
+++ b/Engine/Plugins/Experimental/PlainProps/Source/Private/PlainPropsInternalBind.h
@@ -71,6 +71,7 @@ struct FMemberBinderBase
 
 	void AddRange(TConstArrayView<FRangeBinding> Ranges, FMemberBindType InnermostType, uint32 Offset)
 	{
+		checkSlow(!Ranges.IsEmpty());
 		AddMember(FMemberBindType(Ranges[0].GetSizeType()), Offset);
 
 		for (FRangeBinding Range : Ranges.RightChop(1))
diff --git a/Engine/Plugins/Experimental/PlainProps/Source/Private/PlainPropsLoad.cpp b/Engine/Plugins/Experimental/PlainProps/Source/Private/PlainPropsLoad.cpp
--- a/Engine/Plugins/Experimental/PlainProps/Source/Private/PlainPropsLoad.cpp
+++ b/Engine/Plugins/Experimental/PlainProps/Source/Private/PlainPropsLoad.cpp
@@ -150,17 +150,22 @@ static uint16 CountEnums(const FStructSchema& Schema)
 	return Num;
 }
 
-static bool HasDifferentSupers(const FStructSchema& From, const FSchemaBinding& To, TConstArrayView<FStructId> ToStructIds)
+static bool HasIncompatibleSupers(const FStructSchema& From, const FSchemaBinding& To, TConstArrayView<FStructId> ToStructIds)
 {
 	if (From.Inheritance == ESuper::No)
 	{
 		return To.HasSuper();
 	}
+	else if (From.Inheritance == ESuper::Used)
+	{
+		// From has two super schema ids, dont bother cloning it
+		return true;
+	}
 	else if (To.HasSuper())
 	{
 		FStructId FromSuper = ToStructIds[From.GetSuper().Get().Idx];
 		FStructId ToSuper = To.GetInnerIds()[0].AsStruct();
-		return FromSuper == ToSuper;
+		return FromSuper != ToSuper;
 	}
 
 	return true;
@@ -425,7 +430,7 @@ template<class LoadIdBinder>
 [[nodiscard]] static FLoadStructPlan MakeSchemaLoadPlan(const FStructSchema& From, const FSchemaBinding& To, TConstArrayView<FMemberId> ToMemberIds, TConstArrayView<FStructId> ToStructIds, LoadIdBinder& LoadIds, SubsetByteArray& OutSubsetSchemas)
 {
 	uint16 NumEnums = CountEnums(From);
-	if (From.NumMembers < To.NumMembers || NumEnums || HasDifferentSupers(From, To, ToStructIds))
+	if (From.NumMembers < To.NumMembers || NumEnums || HasIncompatibleSupers(From, To, ToStructIds))
 	{
 		CreateSubsetBindingWithoutEnumIds(From, To, ToMemberIds, NumEnums, LoadIds, OutSubsetSchemas);
 	}
diff --git a/Engine/Plugins/Experimental/PlainProps/Source/Private/PlainPropsPrint.cpp b/Engine/Plugins/Experimental/PlainProps/Source/Private/PlainPropsPrint.cpp
--- a/Engine/Plugins/Experimental/PlainProps/Source/Private/PlainPropsPrint.cpp
+++ b/Engine/Plugins/Experimental/PlainProps/Source/Private/PlainPropsPrint.cpp
@@ -1542,7 +1542,7 @@ void PrintDiff(FUtf8Builder& Out, const FBatchIds& Ids, const FReadDiffPath& Dif
 	// print struct members path with range indices
 	for (FReadDiffNode Node : ReverseIterate(Diff))
 	{
-		if (Node.Name || Node.RangeIdx == ~0u)
+		if (Node.Name || Node.RangeIdx == ~uint64(0))
 		{
 			if (bWasName)
 			{
diff --git a/Engine/Plugins/Experimental/PlainProps/Source/Private/PlainPropsRead.cpp b/Engine/Plugins/Experimental/PlainProps/Source/Private/PlainPropsRead.cpp
--- a/Engine/Plugins/Experimental/PlainProps/Source/Private/PlainPropsRead.cpp
+++ b/Engine/Plugins/Experimental/PlainProps/Source/Private/PlainPropsRead.cpp
@@ -218,7 +218,6 @@ FOptionalMemberId FMemberReader::PeekName() const
 
 FMemberId FMemberReader::PeekNameUnchecked() const
 {
-	checkSlow(!UsesSuper);
 	int32 MemberNameIdx = MemberIdx - UsesSuper;
 	return GetMemberNames()[MemberNameIdx];
 }
diff --git a/Engine/Plugins/Experimental/PlainProps/Source/Public/PlainPropsBuild.h b/Engine/Plugins/Experimental/PlainProps/Source/Public/PlainPropsBuild.h
--- a/Engine/Plugins/Experimental/PlainProps/Source/Public/PlainPropsBuild.h
+++ b/Engine/Plugins/Experimental/PlainProps/Source/Public/PlainPropsBuild.h
@@ -409,7 +409,7 @@ private:
 
 inline void* operator new(std::size_t Size, PlainProps::FScratchAllocator& Scratch)
 {
-	return Scratch.Allocate(Size, 8);
+	return Scratch.Allocate(Size, MIN_ALIGNMENT);
 }
 
 inline void* operator new(std::size_t Size, std::align_val_t Align, PlainProps::FScratchAllocator& Scratch)
diff --git a/Engine/Plugins/Experimental/PlainProps/Source/Public/PlainPropsDeclare.h b/Engine/Plugins/Experimental/PlainProps/Source/Public/PlainPropsDeclare.h
--- a/Engine/Plugins/Experimental/PlainProps/Source/Public/PlainPropsDeclare.h
+++ b/Engine/Plugins/Experimental/PlainProps/Source/Public/PlainPropsDeclare.h
@@ -72,7 +72,7 @@ enum class EMemberPresence : uint8 { RequireAll, AllowSparse };
 
 struct FStructDeclaration
 {
-	mutable std::atomic<int32>		RefCount;
+	mutable std::atomic<int32>		RefCount = {};
 	FDeclId							Id;				// Could be removed, might allow declaration dedup among templated types
 	FOptionalDeclId					Super;
 	uint16							Version;
diff --git a/Engine/Plugins/Experimental/PlainProps/Source/Public/PlainPropsDiff.h b/Engine/Plugins/Experimental/PlainProps/Source/Public/PlainPropsDiff.h
--- a/Engine/Plugins/Experimental/PlainProps/Source/Public/PlainPropsDiff.h
+++ b/Engine/Plugins/Experimental/PlainProps/Source/Public/PlainPropsDiff.h
@@ -57,7 +57,7 @@ struct FReadDiffNode
 	FMemberType				Type;
 	FOptionalStructSchemaId	Struct;
 	FOptionalMemberId		Name;
-	uint64					RangeIdx = ~0u; 
+	uint64					RangeIdx = ~uint64(0); 
 };
 
 struct FReadDiffPath : public TArray<FReadDiffNode, TInlineAllocator<16>> {};
diff --git a/Engine/Plugins/Experimental/PlainProps/Source/Public/PlainPropsSave.h b/Engine/Plugins/Experimental/PlainProps/Source/Public/PlainPropsSave.h
--- a/Engine/Plugins/Experimental/PlainProps/Source/Public/PlainPropsSave.h
+++ b/Engine/Plugins/Experimental/PlainProps/Source/Public/PlainPropsSave.h
@@ -67,7 +67,7 @@ PLAINPROPS_API void							SaveSuperDelta(FMemberBuilder& Out, const void* Super,
 
 inline const void* FBaseline::Get() const
 {
-	// AsLink()->Ptr is untagged, only last link is tagged to mark end and GetSuper() resolves it 
+	// AsLink()->Ptr is untagged, only last link is tagged to mark end and Super() resolves it 
     uint64 Out = (Ptr & TagBit) ? reinterpret_cast<const FDefaultLink*>(Ptr & ~TagBit)->Ptr : Ptr;
     return reinterpret_cast<const void*>(Out);
 }
diff --git a/Engine/Plugins/Experimental/PlainProps/Source/Public/PlainPropsSaveMember.h b/Engine/Plugins/Experimental/PlainProps/Source/Public/PlainPropsSaveMember.h
--- a/Engine/Plugins/Experimental/PlainProps/Source/Public/PlainPropsSaveMember.h
+++ b/Engine/Plugins/Experimental/PlainProps/Source/Public/PlainPropsSaveMember.h
@@ -391,7 +391,7 @@ protected:
 	inline void AddBuiltItem(BuiltItemType Item)
 	{
 #if DO_CHECK
-		check(It < End);
+		check(It + sizeof(Item) <= End);
 #endif
 		reinterpret_cast<BuiltItemType&>(*It) = Item;
 		It += sizeof(Item);
