[PlainProps]
* Fixed console nonunity/nopch compile error when instantiating templated function

diff --git a/Engine/Plugins/Experimental/PlainProps/Source/Public/PlainPropsParse.h b/Engine/Plugins/Experimental/PlainProps/Source/Public/PlainPropsParse.h
--- a/Engine/Plugins/Experimental/PlainProps/Source/Public/PlainPropsParse.h
+++ b/Engine/Plugins/Experimental/PlainProps/Source/Public/PlainPropsParse.h
@@ -2,8 +2,8 @@
 
 #pragma once
 
+#include "Containers/StringView.h"
 #include "PlainPropsTypes.h"
-#include "Containers/StringFwd.h"
 
 namespace PlainProps
 {
