PlainProps rename binding InnermostSchema -> InnermostId

#rnx

diff --git a/Engine/Plugins/Experimental/PlainProps/Source/Private/PlainPropsBind.cpp b/Engine/Plugins/Experimental/PlainProps/Source/Private/PlainPropsBind.cpp
--- a/Engine/Plugins/Experimental/PlainProps/Source/Private/PlainPropsBind.cpp
+++ b/Engine/Plugins/Experimental/PlainProps/Source/Private/PlainPropsBind.cpp
@@ -82,7 +82,7 @@ FLeafMemberBinding FMemberVisitor::GrabLeaf()
 FStructMemberBinding FMemberVisitor::GrabStruct()
 {
 	checkf(!PeekType().AsStruct().IsDynamic, TEXT("Bound structs can't be dynamic"));
-	return { PeekType().AsStruct(), GrabInnerSchema().AsStructBindId(), GrabMemberOffset() };
+	return { PeekType().AsStruct(), GrabInnerId().AsStructBindId(), GrabMemberOffset() };
 }
 
 static bool HasSchema(FMemberBindType Type)
@@ -104,10 +104,10 @@ FRangeMemberBinding FMemberVisitor::GrabRange()
 	const FRangeBinding* RangeBindings = Schema.GetRangeBindings() + InnerRangeIdx;
 	check(MaxSize == RangeBindings[0].GetSizeType());
 	FMemberBindTypeRange InnerTypes = GrabInnerTypes();
-	FOptionalInnerId InnermostSchema = HasSchema(InnerTypes.Last()) ? ToOptional(GrabInnerSchema()) : NoId; 
+	FOptionalInnerId InnermostId = HasSchema(InnerTypes.Last()) ? ToOptional(GrabInnerId()) : NoId; 
 	uint64 Offset = GrabMemberOffset();
 		
-	return { &InnerTypes[0], RangeBindings, static_cast<uint16>(InnerTypes.Num()), InnermostSchema, Offset};
+	return { &InnerTypes[0], RangeBindings, static_cast<uint16>(InnerTypes.Num()), InnermostId, Offset};
 }
 
 void FMemberVisitor::SkipMember()
@@ -135,7 +135,7 @@ FBindId FMemberVisitor::GrabSuper()
 	return Schema.GetInnerIds()[0].AsStructBindId();
 }
 
-FInnerId FMemberVisitor::GrabInnerSchema()
+FInnerId FMemberVisitor::GrabInnerId()
 {
 	check(InnerIdIdx < Schema.NumInnerIds);
 	return Schema.GetInnerIds()[InnerIdIdx++];
@@ -163,7 +163,7 @@ FRangeMemberBinding GetInnerRange(FRangeMemberBinding In)
 {
 	check(In.NumRanges > 1);
 	check(In.InnerTypes[0].IsRange());
-	return { In.InnerTypes + 1, In.RangeBindings + 1, static_cast<uint16>(In.NumRanges - 1), In.InnermostSchema };
+	return { In.InnerTypes + 1, In.RangeBindings + 1, static_cast<uint16>(In.NumRanges - 1), In.InnermostId };
 }
 
 ////////////////////////////////////////////////////////////////////////////////////////////////
@@ -402,7 +402,7 @@ inline FType IndexInnermost(FIdIndexerBase& Ids, FMemberBindType Type, FOptional
 
 FBothType FMemberBinding::IndexParameterName(FIdIndexerBase& Ids) const
 {
-	FType BindType = IndexInnermost(Ids, InnermostType, InnermostSchema);
+	FType BindType = IndexInnermost(Ids, InnermostType, InnermostId);
 	FType DeclType = BindType;
 	for (FRangeBinding Range : ReverseIterate(RangeBindings))
 	{
@@ -418,12 +418,12 @@ FBothType FMemberBinding::IndexParameterName(FIdIndexerBase& Ids) const
 FSchemaBindings::FSchemaBindings(FDebugIds In) : Debug(In) {}
 FSchemaBindings::~FSchemaBindings() {}
 
-static uint16 CountInnerSchemas(TConstArrayView<FMemberBinding> Members)
+static uint16 CountInnerIds(TConstArrayView<FMemberBinding> Members)
 {
 	uint32 Out = 0;
 	for (const FMemberBinding& Member : Members)
 	{
-		Out += !!Member.InnermostSchema;
+		Out += !!Member.InnermostId;
 	}
 	return IntCastChecked<uint16>(Out);
 }
@@ -442,20 +442,20 @@ struct FMemberBinder : FMemberBinderBase
 {
 	FMemberBinder(FSchemaBinding& InSchema)
 	: FMemberBinderBase(InSchema)
-	, InnerSchemaIt(const_cast<FInnerId*>(InSchema.GetInnerIds()))
+	, InnerIdIt(const_cast<FInnerId*>(InSchema.GetInnerIds()))
 	{}
 	
 	~FMemberBinder()
 	{
-		check(Align(InnerSchemaIt, alignof(FRangeBinding)) == (const void*)Schema.GetRangeBindings() || Schema.NumInnerRanges == 0);
+		check(Align(InnerIdIt, alignof(FRangeBinding)) == (const void*)Schema.GetRangeBindings() || Schema.NumInnerRanges == 0);
 	}
 
-	void AddInnerSchema(FInnerId InnermostSchema)
+	void AddInnerId(FInnerId InnermostId)
 	{
-		*InnerSchemaIt++ = InnermostSchema;
+		*InnerIdIt++ = InnermostId;
 	}
 
-	FInnerId* InnerSchemaIt;
+	FInnerId* InnerIdIt;
 };
 
 static void IncRef(FStructDeclarationPtr& Ptr)
@@ -519,7 +519,7 @@ static void DeclareOrValidate(FStructDeclarationPtr& Out, FStructSpec Info, TCon
 FSchemaBinding* FSchemaBindings::Bind(FBindId Id, TConstArrayView<FMemberBinding> Members, FDeclId DeclId) 
 {
 	// Make header, allocate and copy header
-	FSchemaBinding Header = { DeclId, IntCastChecked<uint16>(Members.Num()), CountInnerSchemas(Members), CountRanges(Members) };
+	FSchemaBinding Header = { DeclId, IntCastChecked<uint16>(Members.Num()), CountInnerIds(Members), CountRanges(Members) };
 	FSchemaBinding* Out = new (FMemory::MallocZeroed(Header.CalculateSize())) FSchemaBinding {Header};
 
 	// Write footer
@@ -536,9 +536,9 @@ FSchemaBinding* FSchemaBindings::Bind(FBindId Id, TConstArrayView<FMemberBinding
 			Footer.AddRange(Ranges, Member.InnermostType, IntCastChecked<uint32>(Member.Offset));
 		}
 
-		if (Member.InnermostSchema)
+		if (Member.InnermostId)
 		{
-			Footer.AddInnerSchema(Member.InnermostSchema.Get());
+			Footer.AddInnerId(Member.InnermostId.Get());
 		}
 	}
 
diff --git a/Engine/Plugins/Experimental/PlainProps/Source/Private/PlainPropsDiff.cpp b/Engine/Plugins/Experimental/PlainProps/Source/Private/PlainPropsDiff.cpp
--- a/Engine/Plugins/Experimental/PlainProps/Source/Private/PlainPropsDiff.cpp
+++ b/Engine/Plugins/Experimental/PlainProps/Source/Private/PlainPropsDiff.cpp
@@ -47,7 +47,7 @@ bool DiffMember(FRangeMemberBinding Member, const void* OwnerA, const void* Owne
 		const IItemRangeBinding& ItemBinding = Binding.AsItemBinding();
 		if (InnerKind == EMemberKind::Struct)
 		{
-			bDiff = DiffItemRange(A, B, ItemBinding, Ctx, Member.InnermostSchema.Get().AsStructBindId());
+			bDiff = DiffItemRange(A, B, ItemBinding, Ctx, Member.InnermostId.Get().AsStructBindId());
 		}
 		else if (InnerKind == EMemberKind::Leaf)
 		{
diff --git a/Engine/Plugins/Experimental/PlainProps/Source/Private/PlainPropsInternalBind.h b/Engine/Plugins/Experimental/PlainProps/Source/Private/PlainPropsInternalBind.h
--- a/Engine/Plugins/Experimental/PlainProps/Source/Private/PlainPropsInternalBind.h
+++ b/Engine/Plugins/Experimental/PlainProps/Source/Private/PlainPropsInternalBind.h
@@ -37,10 +37,10 @@ protected: // for unit tests
 
 	uint64						GrabMemberOffset();
 	FMemberBindTypeRange		GrabInnerTypes();
-	FInnerId					GrabInnerSchema();
+	FInnerId					GrabInnerId();
 	FBindId						GrabStructSchema(FStructType Type);
 	FOptionalInnerId			GrabRangeSchema(FMemberType InnermostType);
-	FEnumId						GrabEnumSchema()		{ return GrabInnerSchema().AsEnum(); }
+	FEnumId						GrabEnumSchema()		{ return GrabInnerId().AsEnum(); }
 };
 
 ////////////////////////////////////////////////////////////////////////////////////////////////
diff --git a/Engine/Plugins/Experimental/PlainProps/Source/Private/PlainPropsLoad.cpp b/Engine/Plugins/Experimental/PlainProps/Source/Private/PlainPropsLoad.cpp
--- a/Engine/Plugins/Experimental/PlainProps/Source/Private/PlainPropsLoad.cpp
+++ b/Engine/Plugins/Experimental/PlainProps/Source/Private/PlainPropsLoad.cpp
@@ -271,68 +271,68 @@ struct FMemberLoadBinder : FMemberBinderBase
 {
 	FMemberLoadBinder(FSchemaBinding& In)
 	: FMemberBinderBase(In)
-	, InnerSchemaIt(const_cast<FStructSchemaId*>(FSchemaLoadPlan(In).GetInnerSchemas().GetData()))
+	, InnerIdIt(const_cast<FStructSchemaId*>(FSchemaLoadPlan(In).GetInnerSchemas().GetData()))
 	{}
 	
 	~FMemberLoadBinder()
 	{
-		check(Align(InnerSchemaIt, alignof(FRangeBinding)) == (const void*)Schema.GetRangeBindings() || Schema.NumInnerRanges == 0);
+		check(Align(InnerIdIt, alignof(FRangeBinding)) == (const void*)Schema.GetRangeBindings() || Schema.NumInnerRanges == 0);
 	}
 
-	void AddInnerSchema(FStructSchemaId InnermostSchema)
+	void AddInnerId(FStructSchemaId InnermostId)
 	{
-		*InnerSchemaIt++ = InnermostSchema;
+		*InnerIdIt++ = InnermostId;
 	}
 
-	FStructSchemaId* InnerSchemaIt;
+	FStructSchemaId* InnerIdIt;
 };
 
 
-static void CopyMemberBinding(FLeafMemberBinding Binding, /* in-out */ const FSchemaId*& InnerSchemaIt, FMemberLoadBinder& Out)
+static void CopyMemberBinding(FLeafMemberBinding Binding, /* in-out */ const FSchemaId*& InnerIdIt, FMemberLoadBinder& Out)
 {
-	InnerSchemaIt += (Binding.Leaf.Type == ELeafBindType::Enum); // Skip enum schema
+	InnerIdIt += (Binding.Leaf.Type == ELeafBindType::Enum); // Skip enum schema
 	Out.AddMember(Binding.Leaf.Pack(), static_cast<uint32>(Binding.Offset));
 }
 
 // TODO: Make all below free functions FLoadPlanner member functions and make LoadIds a member variable
 
 template<class LoadIdBinder>
-void CopyMemberBinding(FStructMemberBinding Binding, /* in-out */ const FSchemaId*& InnerSchemaIt, LoadIdBinder& LoadIds, FMemberLoadBinder& Out)
+void CopyMemberBinding(FStructMemberBinding Binding, /* in-out */ const FSchemaId*& InnerIdIt, LoadIdBinder& LoadIds, FMemberLoadBinder& Out)
 {
 	check(!Binding.Type.IsSuper);
 	check(!Binding.Type.IsDynamic); // TODO: Build out dynamic struct support
-	FStructSchemaId LoadId = LoadIds.BindLoadId(static_cast<FStructSchemaId>(*InnerSchemaIt), Binding.Id);
+	FStructSchemaId LoadId = LoadIds.BindLoadId(static_cast<FStructSchemaId>(*InnerIdIt), Binding.Id);
 	Out.AddMember(FMemberBindType(Binding.Type), static_cast<uint32>(Binding.Offset));
-	Out.AddInnerSchema(LoadId);
-	++InnerSchemaIt;
+	Out.AddInnerId(LoadId);
+	++InnerIdIt;
 }
 
 template<class LoadIdBinder>
-void CopyMemberBinding(FRangeMemberBinding Binding, /* in-out */ const FSchemaId*& InnerSchemaIt, LoadIdBinder& LoadIds, FMemberLoadBinder& Out)
+void CopyMemberBinding(FRangeMemberBinding Binding, /* in-out */ const FSchemaId*& InnerIdIt, LoadIdBinder& LoadIds, FMemberLoadBinder& Out)
 {
 	FMemberBindType InnermostType = Binding.InnerTypes[Binding.NumRanges - 1];
 	Out.AddRange(MakeArrayView(Binding.RangeBindings, Binding.NumRanges), InnermostType, static_cast<uint32>(Binding.Offset));
 	if (InnermostType.IsStruct())
 	{
-		FStructSchemaId ReadId = static_cast<FStructSchemaId>(*InnerSchemaIt);
-		FStructSchemaId LoadId = LoadIds.BindLoadId(ReadId, Binding.InnermostSchema.Get().AsStructBindId());
-		Out.AddInnerSchema(LoadId);
-		++InnerSchemaIt;
+		FStructSchemaId ReadId = static_cast<FStructSchemaId>(*InnerIdIt);
+		FStructSchemaId LoadId = LoadIds.BindLoadId(ReadId, Binding.InnermostId.Get().AsStructBindId());
+		Out.AddInnerId(LoadId);
+		++InnerIdIt;
 	}
 	else
 	{
-		InnerSchemaIt += (InnermostType.AsLeaf().Bind.Type == ELeafBindType::Enum); // Skip enum schema
+		InnerIdIt += (InnermostType.AsLeaf().Bind.Type == ELeafBindType::Enum); // Skip enum schema
 	}
 }
 
 template<class LoadIdBinder>
-void CopyMemberBinding(/* in-out */ FMemberVisitor& BindIt, /* in-out */ const FSchemaId*& InnerSchemaIt, LoadIdBinder& LoadIds, FMemberLoadBinder& Out)
+void CopyMemberBinding(/* in-out */ FMemberVisitor& BindIt, /* in-out */ const FSchemaId*& InnerIdIt, LoadIdBinder& LoadIds, FMemberLoadBinder& Out)
 {
 	switch (BindIt.PeekKind())
 	{
-		case EMemberKind::Leaf:		CopyMemberBinding(BindIt.GrabLeaf(), InnerSchemaIt, Out);				break;
-		case EMemberKind::Range:	CopyMemberBinding(BindIt.GrabRange(), InnerSchemaIt, LoadIds, Out);		break;
-		case EMemberKind::Struct:	CopyMemberBinding(BindIt.GrabStruct(), InnerSchemaIt, LoadIds, Out);	break;
+		case EMemberKind::Leaf:		CopyMemberBinding(BindIt.GrabLeaf(), InnerIdIt, Out);				break;
+		case EMemberKind::Range:	CopyMemberBinding(BindIt.GrabRange(), InnerIdIt, LoadIds, Out);		break;
+		case EMemberKind::Struct:	CopyMemberBinding(BindIt.GrabStruct(), InnerIdIt, LoadIds, Out);	break;
 		default:					check(false);															break;
 	}
 }
@@ -356,7 +356,7 @@ static void CreateSubsetBindingWithoutEnumIds(const FStructSchema& From, const F
 	// Copy subset of member bindings...
 	FMemberVisitor ToIt(To);
 	FMemberLoadBinder Footer(*Schema);
-	const FSchemaId* InnerSchemaIt = From.GetInnerSchemas() + bSkipSuperSchema;
+	const FSchemaId* InnerIdIt = From.GetInnerSchemas() + bSkipSuperSchema;
 
 	// ...first the unnamed super member...
 	if (From.Inheritance != ESuper::No)
@@ -370,7 +370,7 @@ static void CreateSubsetBindingWithoutEnumIds(const FStructSchema& From, const F
 		
 			if (!FromType.IsDynamic)
 			{
-				FStructSchemaId ReadId = static_cast<FStructSchemaId>(*InnerSchemaIt);
+				FStructSchemaId ReadId = static_cast<FStructSchemaId>(*InnerIdIt);
 				FStructSchemaId LoadId;
 				if (From.Inheritance == ESuper::Reused)
 				{
@@ -383,8 +383,8 @@ static void CreateSubsetBindingWithoutEnumIds(const FStructSchema& From, const F
 					// i.e. when reading a base of the declared super, needs UsedBindId.
 					LoadId = ReadId;
 				}
-				Footer.AddInnerSchema(LoadId);
-				++InnerSchemaIt;
+				Footer.AddInnerId(LoadId);
+				++InnerIdIt;
 			}
 		}
 	}
@@ -399,9 +399,9 @@ static void CreateSubsetBindingWithoutEnumIds(const FStructSchema& From, const F
 			check(ToNameIt != ToNames.end());
 		}
 		
-		CopyMemberBinding(/* in-out */ ToIt, /* in-out */ InnerSchemaIt, /* in-out */ LoadIds,  /* out */ Footer);	
+		CopyMemberBinding(/* in-out */ ToIt, /* in-out */ InnerIdIt, /* in-out */ LoadIds,  /* out */ Footer);	
 	}
-	check(InnerSchemaIt == From.GetInnerSchemas() + From.NumInnerSchemas);
+	check(InnerIdIt == From.GetInnerSchemas() + From.NumInnerSchemas);
 }
 
 // @pre No enum members
@@ -701,12 +701,12 @@ private:
 				FMemberType Type = NamedTypes[NameIdx];
 				FMemberType InnermostType = Type.IsRange() ? GrabInnerRangeTypes(RangeTypes, /* in-out */ RangeTypeIdx).Last() : Type;
 				++NameIdx;
-				InnerSchemaIt += IsStructOrEnum(InnermostType);
+				InnerIdIt += IsStructOrEnum(InnermostType);
 				
 				if (Name == InName)
 				{
 					check(InnermostType.IsStruct());
-					return static_cast<FStructSchemaId>(InnerSchemaIt[-1]);
+					return static_cast<FStructSchemaId>(InnerIdIt[-1]);
 				}
 			}
 		}
@@ -718,13 +718,13 @@ private:
 		
 		int32								NameIdx = 0;
 		int32								RangeTypeIdx = 0;
-		const FSchemaId*					InnerSchemaIt;
+		const FSchemaId*					InnerIdIt;
 
 		FInnerStructSchemaIterator(const FStructSchema& Schema, bool bSkipSuper)
 		: NamedTypes(FStructSchema::GetMemberTypes(Schema.Footer) + bSkipSuper)
 		, Names(Schema.GetMemberNames())
 		, RangeTypes(Schema.GetRangeTypes())
-		, InnerSchemaIt(Schema.GetInnerSchemas() + /* super schema */ (bSkipSuper && !Schema.GetMemberTypes()[0].AsStruct().IsDynamic))
+		, InnerIdIt(Schema.GetInnerSchemas() + /* super schema */ (bSkipSuper && !Schema.GetMemberTypes()[0].AsStruct().IsDynamic))
 		{}
 	};
 
@@ -1345,18 +1345,18 @@ void FSchemaLoadHandle::GetInnerLoadIds(TArrayView<FOptionalSchemaId> Out) const
 	}
 	else
 	{
-		const FSchemaId* InnerSchemaIt = ReadSchema.GetInnerSchemas();
+		const FSchemaId* InnerIdIt = ReadSchema.GetInnerSchemas();
 		for (FMemberType Member : MemberTypes)
 		{
 			FMemberType Innermost = Member.IsRange() ? GrabInnerRangeTypes(RangeTypes, /* in-out */ RangeTypeIdx).Last() : Member;
 			if (Innermost.IsStruct())
 			{
-				*OutIt++ = *InnerSchemaIt++;
+				*OutIt++ = *InnerIdIt++;
 			}
 			else
 			{
 				*OutIt++ = NoId;
-				InnerSchemaIt += Innermost.AsLeaf().Type == ELeafType::Enum;
+				InnerIdIt += Innermost.AsLeaf().Type == ELeafType::Enum;
 			}
 		}
 	}
diff --git a/Engine/Plugins/Experimental/PlainProps/Source/Private/PlainPropsSave.cpp b/Engine/Plugins/Experimental/PlainProps/Source/Private/PlainPropsSave.cpp
--- a/Engine/Plugins/Experimental/PlainProps/Source/Private/PlainPropsSave.cpp
+++ b/Engine/Plugins/Experimental/PlainProps/Source/Private/PlainPropsSave.cpp
@@ -254,7 +254,7 @@ inline const FStructDeclaration& GetDeclaration(const FSaveContext& Ctx, FBindId
 	{
 	case EMemberKind::Leaf:		return SaveLeafRange(Range, ItemBinding, Ctx, GetItemWidth(InnerType.AsLeaf()));
 	case EMemberKind::Range:	return SaveNonLeafRange<FInternalNestedRangeSaver>(Range, ItemBinding, Ctx, GetInnerRange(Member));
-	case EMemberKind::Struct:	return SaveStructRange(Range, ItemBinding, Ctx, Member.InnermostSchema.Get().AsStructBindId());
+	case EMemberKind::Struct:	return SaveStructRange(Range, ItemBinding, Ctx, Member.InnermostId.Get().AsStructBindId());
 	}
 
 	check(false);
@@ -296,7 +296,7 @@ inline const FStructDeclaration& GetDeclaration(const FSaveContext& Ctx, FBindId
 [[nodiscard]] static FMemberSchema CreateRangeSchema(FScratchAllocator& Scratch, FRangeMemberBinding Member)
 {
 	FMemberType* InnerRangeTypes = CreateInnerRangeTypes(Scratch, Member.NumRanges, Member.InnerTypes);
-	return { FMemberType(Member.RangeBindings[0].GetSizeType()), ToMemberType(Member.InnerTypes[0]), Member.NumRanges, Member.InnermostSchema, InnerRangeTypes };
+	return { FMemberType(Member.RangeBindings[0].GetSizeType()), ToMemberType(Member.InnerTypes[0]), Member.NumRanges, Member.InnermostId, InnerRangeTypes };
 }
 
 static void SaveMember(FMemberBuilder& Out, const void* Struct, FMemberId Name, const FSaveContext& Ctx, FLeafMemberBinding Member)
@@ -377,7 +377,7 @@ static bool DiffItem(const uint8* A, const uint8* B, const FSaveContext& Ctx, FR
 	{
 	case EMemberKind::Leaf:		return DiffItemRange(A, B, ItemBinding, Ctx, SizeOf(GetItemWidth(InnerType.AsLeaf())));
 	case EMemberKind::Range:	return DiffItemRange(A, B, ItemBinding, Ctx, GetInnerRange(Member));
-	case EMemberKind::Struct:	return DiffItemRange(A, B, ItemBinding, Ctx, Member.InnermostSchema.Get().AsStructBindId());
+	case EMemberKind::Struct:	return DiffItemRange(A, B, ItemBinding, Ctx, Member.InnermostId.Get().AsStructBindId());
 	}
 
 	check(false);
diff --git a/Engine/Plugins/Experimental/PlainProps/Source/Public/PlainPropsBind.h b/Engine/Plugins/Experimental/PlainProps/Source/Public/PlainPropsBind.h
--- a/Engine/Plugins/Experimental/PlainProps/Source/Public/PlainPropsBind.h
+++ b/Engine/Plugins/Experimental/PlainProps/Source/Public/PlainPropsBind.h
@@ -212,7 +212,7 @@ struct FRangeMemberBinding
 	const FMemberBindType*	InnerTypes;
 	const FRangeBinding*	RangeBindings;
 	uint16					NumRanges; // At least 1, >1 for nested ranges
-	FOptionalInnerId		InnermostSchema;
+	FOptionalInnerId		InnermostId;
 	SIZE_T					Offset;
 };
 
@@ -616,10 +616,10 @@ struct FMemberBinding
 
 	uint64							Offset;
 	FMemberBindType					InnermostType;		// Always Leaf or Struct
-	FOptionalInnerId				InnermostSchema;	// Enum or struct schema
+	FOptionalInnerId				InnermostId;		// Enum or struct bind id
 	TConstArrayView<FRangeBinding>	RangeBindings;		// Non-empty -> Range
 
-	// @pre InnermostSchema isn't type-erased / lowered
+	// @pre InnermostId isn't type-erased / lowered
 	PLAINPROPS_API FBothType		IndexParameterName(FIdIndexerBase& Ids) const;
 };
 
@@ -1126,7 +1126,7 @@ template<LeafType Type, class Runtime>
 FMemberBinding BindMember(uint64 Offset, FMemberSpec& OutSpec)
 {
 	FMemberBinding Out(Offset);
-	Out.InnermostType = BindInnermostType<Type, typename Runtime::Ids>(Out.InnermostSchema, OutSpec);
+	Out.InnermostType = BindInnermostType<Type, typename Runtime::Ids>(Out.InnermostId, OutSpec);
 	return Out;
 }
 
@@ -1140,7 +1140,7 @@ FMemberBinding BindMember(uint64 Offset, FMemberSpec& OutSpec)
 	if constexpr (!std::is_void_v<CustomBinding>)
 	{
 		FBothStructId Both = BindCustomStructOnce<Type, CustomBinding, Runtime>();
-		Out.InnermostType = BindInnerStruct(Out.InnermostSchema, OutSpec, Both);
+		Out.InnermostType = BindInnerStruct(Out.InnermostId, OutSpec, Both);
 	}
 	else
 	{
@@ -1151,7 +1151,7 @@ FMemberBinding BindMember(uint64 Offset, FMemberSpec& OutSpec)
 			using InnermostType = typename TInnermostType<RangeBinding, NumRangeBindings>::Type;
 
 			Out.RangeBindings = GetRangeBindings<RangeBinding, Ids, NumRangeBindings>();
-			Out.InnermostType = BindInnermostType<InnermostType, Ids>(Out.InnermostSchema, OutSpec);
+			Out.InnermostType = BindInnermostType<InnermostType, Ids>(Out.InnermostId, OutSpec);
 			
 			for (FRangeBinding Range : ReverseIterate(Out.RangeBindings))
 			{
@@ -1161,7 +1161,7 @@ FMemberBinding BindMember(uint64 Offset, FMemberSpec& OutSpec)
 		else
 		{
 			FBothStructId Both = GetStructBothId<typename Runtime::Ids, Type>();
-			Out.InnermostType = BindInnerStruct(Out.InnermostSchema, OutSpec, Both);
+			Out.InnermostType = BindInnerStruct(Out.InnermostId, OutSpec, Both);
 		}
 	}
 	
diff --git a/Engine/Plugins/Experimental/PlainPropsUObject/Source/Private/PlainPropsUObjectBindings.cpp b/Engine/Plugins/Experimental/PlainPropsUObject/Source/Private/PlainPropsUObjectBindings.cpp
--- a/Engine/Plugins/Experimental/PlainPropsUObject/Source/Private/PlainPropsUObjectBindings.cpp
+++ b/Engine/Plugins/Experimental/PlainPropsUObject/Source/Private/PlainPropsUObjectBindings.cpp
@@ -882,13 +882,13 @@ inline uint32 HashRangeBindings(TConstArrayView<FRangeBinding> In)
 
 inline uint32 HashSkipOffset(FMemberBinding In)
 {
-	uint32 Out = HashCombineFast(GetTypeHash(In.InnermostSchema), GetTypeHash(In.InnermostType));
+	uint32 Out = HashCombineFast(GetTypeHash(In.InnermostId), GetTypeHash(In.InnermostType));
 	return In.RangeBindings.IsEmpty() ? Out : HashCombineFast(Out, HashRangeBindings(In.RangeBindings));
 }
 
 inline bool EqSkipOffset(FMemberBinding A, FMemberBinding B)
 {
-	return A.InnermostType == B.InnermostType && A.InnermostSchema == B.InnermostSchema && Algo::Compare(A.RangeBindings, B.RangeBindings);
+	return A.InnermostType == B.InnermostType && A.InnermostId == B.InnermostId && Algo::Compare(A.RangeBindings, B.RangeBindings);
 }
 
 // Helper to cache various property bindings instead of a TMap KeyFunc
@@ -1508,7 +1508,7 @@ struct TLeafRangeSerializer
 
 	explicit TLeafRangeSerializer(FMemberBinding In)
 	: InnerType(ToLeafType(In.InnermostType.AsLeaf()))
-	, EnumId(In.InnermostSchema)
+	, EnumId(In.InnermostId)
 	{
 		check(In.RangeBindings.IsEmpty());
 		check(InnerType.AsLeaf().Width == WidthOf(Size));
@@ -1554,7 +1554,7 @@ struct FStructRangeSerializer
 
 	explicit FStructRangeSerializer(FMemberBinding Item)
 	: InnerType(Item.InnermostType.AsStruct())
-	, SaveId(Item.InnermostSchema.Get().AsStruct())
+	, SaveId(Item.InnermostId.Get().AsStruct())
 	{
 		check(Item.RangeBindings.IsEmpty());
 	}
@@ -1592,7 +1592,7 @@ struct FNestedRangeSerializer
 	TArray<FRangeBinding, TInlineAllocator<2>>		InnerBindings;
 
 	explicit FNestedRangeSerializer(FMemberBinding Item)
-	: InnermostSaveId(Item.InnermostSchema)
+	: InnermostSaveId(Item.InnermostId)
 	, NumInners(IntCastChecked<uint16>(1 + Item.RangeBindings.Num()))
 	, InnerBindings(Item.RangeBindings)
 	{
@@ -2630,7 +2630,7 @@ class FMapBindings
 			FType DeclType = { GUE.Scopes.CoreUObject, FTypenameId(GUE.Names.MakeParametricTypeId(GUE.Typenames.Map, DeclParams)) };
 			BothPair = { GUE.Names.IndexBindId(BindType), GUE.Names.IndexDeclId(DeclType) };
 
-			FBindId PairId = Members.Pair.InnermostSchema.Get().AsStructBindId();
+			FBindId PairId = Members.Pair.InnermostId.Get().AsStructBindId();
 			if (BothKey.IsLowered())
 			{
 				FBindId KeyId = GUE.Names.IndexBindId(BothKey.BindType);
@@ -2649,7 +2649,7 @@ class FMapBindings
 			checkf(!Members.Key.IndexParameterName(GUE.Names).IsLowered(), TEXT("Only range-bound keys or values should be type-erased / lowered"));
 			checkf(!Members.Value.IndexParameterName(GUE.Names).IsLowered(), TEXT("Only range-bound keys or values should be type-erased / lowered"));
 
-			FParametricTypeId PairTypename = GUE.Names.Resolve(Members.Pair.InnermostSchema.Get().AsStruct()).Name.AsParametric();
+			FParametricTypeId PairTypename = GUE.Names.Resolve(Members.Pair.InnermostId.Get().AsStruct()).Name.AsParametric();
 			TConstArrayView<FType> Params = GUE.Names.Resolve(PairTypename).GetParameters();
 			FType Type = { GUE.Scopes.CoreUObject, FTypenameId(GUE.Names.MakeParametricTypeId(GUE.Typenames.Map, Params)) };
 			FDualStructId Dual{GUE.Names.IndexStruct(Type)};
@@ -2727,12 +2727,12 @@ static FMapBindings GMaps;
 
 ////////////////////////////////////////////////////////////////////////////////////////////////
 
-// @pre Binding.InnermostSchema isn't lowered
+// @pre Binding.InnermostId isn't lowered
 static FMemberSpec ToSpec(FMemberBinding Binding)
 {
 	FMemberSpec Out = Binding.InnermostType.IsLeaf() 
-		? FMemberSpec(ToLeafType(Binding.InnermostType.AsLeaf()), ToOptionalEnum(Binding.InnermostSchema))
-		: FMemberSpec(Binding.InnermostSchema.Get().AsStructDeclId());
+		? FMemberSpec(ToLeafType(Binding.InnermostType.AsLeaf()), ToOptionalEnum(Binding.InnermostId))
+		: FMemberSpec(Binding.InnermostId.Get().AsStructDeclId());
 
 	for (FRangeBinding RangeBinding : Binding.RangeBindings)
 	{
@@ -2787,7 +2787,7 @@ public:
 	inline FMemberBinding Bind(FMemberBinding Key, FMemberBinding Value)
 	{
 		FMemberBinding Out(0);
-		Out.InnermostSchema = FInnerId(BindImpl(FPair{Key, Value}));
+		Out.InnermostId = FInnerId(BindImpl(FPair{Key, Value}));
 		Out.InnermostType = DefaultStructBindType;
 		return Out;
 	}
@@ -3140,7 +3140,7 @@ private:
 
 		FMemberBinding Out(Property->GetOffset_ForInternal());
 		Out.InnermostType = Inner.InnermostType;
-		Out.InnermostSchema = Inner.InnermostSchema;
+		Out.InnermostId = Inner.InnermostId;
 		Out.RangeBindings = AllocateRangeBindings(RangeBinding, Inner.RangeBindings);
 		return Out;
 	}
@@ -3148,7 +3148,7 @@ private:
 	inline FMemberBinding BindAsStruct(FProperty* Property, FBindId Id)
 	{
 		FMemberBinding Out(Property->GetOffset_ForInternal());
-		Out.InnermostSchema = FInnerId(Id);
+		Out.InnermostId = FInnerId(Id);
 		Out.InnermostType = DefaultStructBindType;
 		return Out;
 	}
@@ -3178,7 +3178,7 @@ private:
 	inline FMemberBinding BindEnum(FEnumProperty* Property)
 	{
 		FMemberBinding Out(Property->GetOffset_ForInternal());
-		Out.InnermostSchema = FInnerId(GUE.Names.IndexEnum(IndexType(Property->GetEnum())));
+		Out.InnermostId = FInnerId(GUE.Names.IndexEnum(IndexType(Property->GetEnum())));
 		FUnpackedLeafType Leaf = {ELeafType::Enum, WidthOf(Property->GetElementSize())};
 		Out.InnermostType = FMemberBindType(Leaf);
 		return Out;
@@ -3201,7 +3201,7 @@ private:
 		bool bIntS = HasAny<CASTCLASS_FInt8Property | CASTCLASS_FInt16Property | CASTCLASS_FIntProperty | CASTCLASS_FInt64Property>(Flags);
 		if (HasAny<CASTCLASS_FByteProperty>(Flags))
 		{
-			FLeafBindType Leaf = BindByte(static_cast<FByteProperty*>(Property), /* out enum */ Out.InnermostSchema);
+			FLeafBindType Leaf = BindByte(static_cast<FByteProperty*>(Property), /* out enum */ Out.InnermostId);
 			Out.InnermostType = FMemberBindType(Leaf);
 		}
 		else
@@ -3223,7 +3223,7 @@ private:
 		FMemberBinding Key = BindSingleProperty(Property->KeyProp);
 		FMemberBinding Value = BindSingleProperty(Property->ValueProp);
 		FMemberBinding Pair = GPairs.Bind(Key, Value);
-		FBindId PairId = Pair.InnermostSchema.Get().AsStructBindId();
+		FBindId PairId = Pair.InnermostId.Get().AsStructBindId();
 
 		bool bFreezable = EnumHasAnyFlags(Property->MapFlags, EMapPropertyFlags::UsesMemoryImageAllocator);
 		return bFreezable	? BindAsRange(Property, GMaps.BindFreezable(Property, PairId, {Key, Value, Pair}), Pair)
@@ -3378,7 +3378,7 @@ private:
 			GUE.Schemas.BindStruct(StaticArrayId, ElemBindings, Spec);
 			
 			Out.InnermostType = DefaultStructBindType;
-			Out.InnermostSchema = FInnerId(StaticArrayId);
+			Out.InnermostId = FInnerId(StaticArrayId);
 			Out.RangeBindings = {};
 		}
 	
@@ -3394,7 +3394,7 @@ void FPropertyBinder::BindSuper(FDeclId SuperId)
 
 	FMemberBinding Member;
 	Member.InnermostType = SuperStructBindType;
-	Member.InnermostSchema = FInnerId(SuperId);
+	Member.InnermostId = FInnerId(SuperId);
 	Members.Emplace(Member);
 }
 
