PlainProps split parts of Bind.h into BindCtti.h

#rnx

diff --git a/Engine/Plugins/Experimental/PlainProps/Source/Private/PlainPropsBind.cpp b/Engine/Plugins/Experimental/PlainProps/Source/Private/PlainPropsBind.cpp
--- a/Engine/Plugins/Experimental/PlainProps/Source/Private/PlainPropsBind.cpp
+++ b/Engine/Plugins/Experimental/PlainProps/Source/Private/PlainPropsBind.cpp
@@ -8,6 +8,7 @@
 #include "PlainPropsInternalBuild.h"
 #include "PlainPropsInternalFormat.h"
 #include "PlainPropsInternalRead.h"
+#include "PlainPropsSpecify.h"
 #include "Templates/RefCounting.h"
 
 namespace PlainProps
diff --git a/Engine/Plugins/Experimental/PlainProps/Source/Public/PlainPropsBind.h b/Engine/Plugins/Experimental/PlainProps/Source/Public/PlainPropsBind.h
--- a/Engine/Plugins/Experimental/PlainProps/Source/Public/PlainPropsBind.h
+++ b/Engine/Plugins/Experimental/PlainProps/Source/Public/PlainPropsBind.h
@@ -4,16 +4,11 @@
 
 #include "Containers/ArrayView.h"
 #include "Containers/Set.h"
-#include "Containers/StringView.h"
 #include "Memory/MemoryFwd.h"
 #include "Memory/MemoryView.h"
-#include "PlainPropsCtti.h"
 #include "PlainPropsDeclare.h"
-#include "PlainPropsTypename.h"
 #include "PlainPropsRead.h"
-#include "PlainPropsSpecify.h"
 #include "PlainPropsTypes.h"
-#include <tuple>
 
 namespace PlainProps 
 {
@@ -34,16 +29,10 @@ class FScratchAllocator;
 struct FStructLoadView;
 class FRangeBinding;
 struct FSaveContext;
-struct FTypedRange;
-struct IItemRangeBinding;
 template<class T> class TIdIndexer;
 
 ////////////////////////////////////////////////////////////////////////////////////////////////
 
-inline FAnsiStringView ToAnsiView(std::string_view Str) { return FAnsiStringView(Str.data(), Str.length()); }
-
-////////////////////////////////////////////////////////////////////////////////////////////////
-
 enum class ELeafBindType : uint8 { Bool, IntS, IntU, Float, Hex, Enum, Unicode, BitfieldBool };
 
 inline static constexpr ELeafBindType ToLeafBindType(ELeafType Type)
@@ -347,31 +336,6 @@ public:
 	explicit FCustomBindingsOverlay(const FCustomBindings& Under) : FCustomBindings(&Under), Underlay(Under) {}
 };
 
-template<typename T>
-struct TCustomBind { using Type = void; };
-
-template<class T>
-struct TCustomDeltaBind { using Type = void; };
-
-template<typename T>
-using CustomBind = typename TCustomBind<T>::Type;
-
-template<typename T>
-using CustomDeltaBind = typename TCustomDeltaBind<T>::Type;
-
-////////////////////////////////////////////////////////////////////////////////////////////////
-
-template<typename T>
-struct TOccupancyOf
-{
-	static constexpr EMemberPresence Value = EMemberPresence::AllowSparse;
-};
-
-struct FRequireAll
-{
-	static constexpr EMemberPresence Value = EMemberPresence::RequireAll;
-};
-
 ////////////////////////////////////////////////////////////////////////////////////////////////
 
 // Request from loading layer to IItemRangeBinding to allocate and construct items
@@ -589,12 +553,6 @@ public:
 	inline bool	operator==(FRangeBinding O) const { return Handle == O.Handle; }
 };
 
-template<typename T>
-struct TRangeBind{ using Type = void; };
-
-template<typename T>
-using RangeBind = typename TRangeBind<T>::Type;
-
 ////////////////////////////////////////////////////////////////////////////////////////////////
 
 struct FBothType
@@ -683,544 +641,6 @@ struct FBindDeclarations final : IDeclarations
 	PLAINPROPS_API virtual FDeclId						Lower(FBindId Id) const override;
 };
 
-////////////////////////////////////////////////////////////////////////////////////////////////
-
-template<class Ids, class Typename>
-FScopeId IndexNamespaceId() requires (Typename::Namespace.size() > 0)
-{
-	// Opt: Make cached GetNamespaceId(), either via new namespace CTTI types (maybe PP_REFLECT_NAMESPACE)
-	//		or some compile time string template parameters, perhaps a variadic template taking any number of chars
-	return Ids::IndexScope(ToAnsiView(Typename::Namespace));
-}
-
-template<class Ids, class Typename>
-FScopeId IndexNamespaceId()
-{
-	return NoId;
-}
-
-template<ETypename Kind, class Typename>
-constexpr std::string_view SelectStructName()
-{	
-	if constexpr (Kind == ETypename::Bind && ExplicitBindName<Typename>)
-	{
-		return Typename::BindName;
-	}
-	else
-	{
-		return Typename::DeclName;	
-	}
-}
-
-template<class Ids, ETypename Kind, class Typename>
-FType IndexStructName()
-{
-	FType BaseName = { IndexNamespaceId<Ids, Typename>(), 
-						 FTypenameId(Ids::IndexTypename(ToAnsiView(SelectStructName<Kind, Typename>()))) };
-	
-	if constexpr (ParametricName<Typename>)
-	{
-		return IndexParametricType<Ids, Kind>(BaseName, (typename Typename::Parameters*)nullptr);
-	}
-	else
-	{
-		return BaseName;
-	}
-}
-
-template<class Ids, class Typename>
-FBothStructId IndexStructBothId(FType DeclName = IndexStructName<Ids, ETypename::Decl, Typename>())
-{	
-	FBothStructId Out;
-	Out.DeclId = FDeclId(Ids::IndexStruct(DeclName));
-	Out.BindId = UpCast(Out.DeclId);
-
-	if constexpr (ExplicitBindName<Typename> || ParametricName<Typename>)
-	{
-		FType BindName = IndexStructName<Ids, ETypename::Bind, Typename>();
-		Out.BindId = BindName != DeclName ? FBindId(Ids::IndexStruct(BindName)) : UpCast(Out.DeclId);
-	}
-
-	return Out;
-}
-
-template<class Ids, class Typename>
-FDualStructId IndexStructDualId(FType Name = IndexStructName<Ids, ETypename::Decl, Typename>()) requires (!ExplicitBindName<Typename>)
-{	
-	return FDualStructId(Ids::IndexStruct(Name));
-}
-
-// Cached by function static
-template<class Ids, typename Struct>
-FDeclId GetStructDeclId()
-{
-	static FDeclId Id = FDeclId(Ids::IndexStruct(IndexStructName<Ids, ETypename::Decl, TTypename<Struct>>()));
-	return Id;
-}
-
-// Cached by function static
-template<class Ids, typename Struct>
-FBindId GetStructBindId()
-{
-	static FBindId Id = FBindId(Ids::IndexStruct(IndexStructName<Ids, ETypename::Bind, TTypename<Struct>>()));
-	return Id;
-}
-
-// Cached by function static
-template<class Ids, typename Struct>
-FBothStructId GetStructBothId()
-{
-	static FBothStructId Id = IndexStructBothId<Ids, TTypename<Struct>>();
-	return Id;
-}
-
-// Cached by function static
-template<class Ids, typename Struct>
-FDualStructId GetStructDualId()
-{
-	static FDualStructId Id = IndexStructDualId<Ids, TTypename<Struct>>();
-	return Id;
-}
-
-template<class Ids, typename Ctti>
-FType IndexCttiName()
-{
-	FTypenameId Name{Ids::IndexTypename(Ctti::Name)};
-	FScopeId Namespace = NoId;
-	if constexpr (Ctti::Namespace[0] != '\0')
-	{
-		Namespace = Ids::IndexScope(ToAnsiView(Ctti::Namespace));
-	}
-	return { Namespace, Name };
-}
-
-// Cached by function static
-template<class Ids, Enumeration Enum>
-FEnumId GetEnumId()
-{
-	static FEnumId Id = Ids::IndexEnum(IndexCttiName<Ids, CttiOf<Enum>>());
-	return Id;
-}
-
-template<class Ids, Arithmetic T>
-FType IndexArithmeticName()
-{
-	static constexpr FUnpackedLeafType Leaf = ReflectArithmetic<T>;
-	return { NoId, FTypenameId(Ids::IndexTypename(ToAnsiView(ArithmeticName<Leaf.Type, Leaf.Width>))) };
-}
-
-template<class Ids, ETypename, Arithmetic Leaf>
-FType IndexParameterName()
-{
-	return IndexArithmeticName<Ids, Leaf>();
-}
-
-template<class Ids, ETypename, Enumeration Enum>
-FType IndexParameterName()
-{
-	return GetEnumId<Ids, Enum>();
-}
-
-template<class Ids, ETypename Kind, typename T>
-FType IndexParameterName()
-{
-	using RangeBinding = RangeBind<T>;
-	if constexpr (std::is_void_v<RangeBinding>)
-	{
-		return IndexStructName<Ids, Kind, TTypename<T>>();
-	}
-	else
-	{
-		FType ItemParam = IndexParameterName<Ids, Kind, typename RangeBinding::ItemType>();
-		FType SizeParam = Ids::GetIndexer().MakeRangeParameter(RangeSizeOf(typename RangeBinding::SizeType{}));
-
-		if constexpr (Kind == ETypename::Decl)
-		{
-			// Type-erase range type
-			return Ids::GetIndexer().MakeAnonymousParametricType({ItemParam, SizeParam});
-		}
-		else
-		{
-			using Typename = TTypename<T>;
-			FType RangeBindName = { IndexNamespaceId<Ids, Typename>(), FTypenameId(Ids::IndexTypename(ToAnsiView(RangeBinding::BindName))) };
-			return Ids::GetIndexer().MakeParametricType(RangeBindName, {ItemParam, SizeParam});
-		}
-	}
-}
-
-template<class Ids, ETypename Kind, typename... Ts>
-FType IndexParametricType(FType TemplatedType, const std::tuple<Ts...>*)
-{
-	FType Parameters[] = { (IndexParameterName<Ids, Kind, Ts>())... };
-	return Ids::GetIndexer().MakeParametricType(TemplatedType, Parameters);
-}
-
-////////////////////////////////////////////////////////////////////////////////////////////////
-
-// Works around C++'s lack of templated nullary constructors
-template<class Ids>
-struct TInit {};
-
-////////////////////////////////////////////////////////////////////////////////////////////////
-
-using InnerStructArray = TArray<FInnerStruct, TInlineAllocator<8>>;
-
-struct FCustomInit
-{
-	UE_NONCOPYABLE(FCustomInit);
-	FCustomInit() = default;
-
-	InnerStructArray				LoweredInners;
-
-	// Only needed for inner structs that might get type-erased/lowered
-	void RegisterInnerStruct(FBothStructId InnerId, TConstArrayView<FMemberId> Names)
-	{
-		if (InnerId.IsLowered())
-		{
-			for (FMemberId Name : Names)
-			{
-				LoweredInners.Add({Name, InnerId.BindId});
-			}
-		}
-	}
-};
-
-// Helps custom binding constructors create ids and register types that might get type-erased 
-template<class Ids>
-struct TCustomInit : FCustomInit, TInit<Ids>
-{};
-
-// Constructor API for static custom bindings to create ids, register type-erased inner structs
-// and specify member types, inheritance and version 
-template<class Ids, uint32 N>
-struct TCustomSpecifier : TCustomInit<Ids>
-{
-	FOptionalStructId				Super;
-	uint16							Version = 0;
-	FMemberSpec						Members[N];
-
-	template <typename... Ts>
-    void SetMembers(Ts&&... Specs) requires (N == sizeof...(Specs))
-    {
-		FMemberSpec* It = &Members[0];
-		((*It++ = FMemberSpec(Forward<Ts>(Specs))), ...);
-    }
-
-	void FillMembers(FMemberSpec Same)
-	{
-		for (FMemberSpec& Member : Members)
-		{
-			Member = Same;
-		}
-	}
-};
-
-// Workaround for clang requiring the definition, make it dependent on any other type
-template<typename>
-struct TDependent
-{
-	using Baseline = FBaseline;
-	using StructLoadView = FStructLoadView;
-};
-
-// Scoped custom binding for native type
-template<class T, class Runtime, typename CustomBinding = CustomBind<T>>
-struct TScopedStructBinding : FBothStructId, CustomBinding
-{
-	using Ids = typename Runtime::Ids;
-	using Typename = TCustomTypename<CustomBinding>::Type;
-	static constexpr uint32 N = sizeof(CustomBinding::MemberIds) / sizeof(CustomBinding::MemberIds[0]);
-	static constexpr bool bFastLoad = requires(FMemcpyLoadPlan& O) { CustomBinding::Plan(O); };
-
-	explicit TScopedStructBinding(FType DeclName = IndexStructName<Ids, ETypename::Decl, Typename>())
-	: TScopedStructBinding(IndexStructBothId<Ids, Typename>(DeclName))
-	{}
-
-	explicit TScopedStructBinding(FDualStructId Dual)
-	: TScopedStructBinding(FBothStructId{Dual, Dual})
-	{}
-
-	explicit TScopedStructBinding(FBothStructId Both)
-	: TScopedStructBinding(Both, TCustomSpecifier<Ids, N>())
-	{}
-
-	TScopedStructBinding(FBothStructId Both, TCustomSpecifier<Ids, N>&& Init)
-	: FBothStructId(Both)
-	, CustomBinding(/* out */ Init)
-	{
-		FStructSpec Decl = { DeclId, /* todo */ NoId, Init.Version, TOccupancyOf<T>::Value, CustomBinding::MemberIds, Init.Members };
-		Runtime::GetCustoms().BindStruct(BindId, *this, Decl, Init.LoweredInners);
-	}
-
-	~TScopedStructBinding()
-	{
-		Runtime::GetCustoms().DropStruct(BindId);
-	}
-
-	virtual void SaveCustom(FMemberBuilder& Dst, const void* Src, TDependent<Ids>::Baseline Base, const FSaveContext& Ctx) override
-	{
-		const void* Default = Base.Get();
-		CustomBinding::Save(Dst, *static_cast<const T*>(Src), static_cast<const T*>(Default), Ctx);
-	}
-
-	virtual void LoadCustom(void* Dst, TDependent<Ids>::StructLoadView Src, ECustomLoadMethod Method) const override 
-	{
-		if constexpr (bFastLoad)
-		{
-			unimplemented();
-		}
-		else
-		{
-			CustomBinding::Load(*static_cast<T*>(Dst), Src, Method);
-		}
-	}
-	
-	virtual bool DiffCustom(const void* A, const void* B, const FBindContext& Ctx) const override
-	{
-		return CustomBinding::Diff(*static_cast<const T*>(A), *static_cast<const T*>(B), Ctx);
-	}
-
-	virtual bool DiffCustom(const void* A, const void* B, FDiffContext& Ctx) const override
-	{
-		return CustomBinding::Diff(*static_cast<const T*>(A), *static_cast<const T*>(B), Ctx);
-	}
-
-	virtual void PlanCustom(FMemcpyLoadPlan& Out) const override
-	{
-		if constexpr (bFastLoad)
-		{
-			CustomBinding::Plan(Out);
-		}
-	}
-};
-
-// Specialization for schemabound structs
-template<class T, class Runtime>
-struct TScopedStructBinding<T, Runtime, void> : FBothStructId
-{
-	using Ids = typename Runtime::Ids;
-	using Typename = TTypename<T>;
-	
-	TScopedStructBinding(FBothStructId Both = IndexStructBothId<Ids, Typename>() )
-	: FBothStructId(Both)
-	{
-		BindNativeStruct<CttiOf<T>, Runtime>(/* out */ Runtime::GetSchemas(), Both, TOccupancyOf<T>::Value);
-	}
-
-	~TScopedStructBinding()
-	{
-		Runtime::GetSchemas().DropStruct(BindId);
-	}
-};
-
-template<class Struct, typename CustomBinding, class Runtime>
-FBothStructId BindCustomStructOnce()
-{
-	static TScopedStructBinding<Struct, Runtime, CustomBinding> Instance;
-	return Instance;
-}
-
-inline constexpr FMemberBindType DefaultStructBindType = FMemberBindType(FStructType{EMemberKind::Struct, /* IsDynamic */ 0, /* IsSuper */ 0});
-inline constexpr FMemberBindType SuperStructBindType = FMemberBindType(FStructType{EMemberKind::Struct, /* IsDynamic */ 0, /* IsSuper */ 1});
-
-inline FMemberBindType BindInnerStruct(FOptionalInnerId& OutBindId, FMemberSpec& OutSpec, FBothStructId Both)
-{
-	OutBindId = FInnerId(Both.BindId);
-	OutSpec = FMemberSpec(Both.DeclId);
-	return DefaultStructBindType;
-}
-
-inline FMemberBindType BindInnerLeaf(FOptionalInnerId& OutId, FMemberSpec& OutSpec, FUnpackedLeafType Leaf, FOptionalEnumId Id)
-{
-	OutId = ToOptionalInner(Id);
-	OutSpec = FMemberSpec(Leaf, Id);
-	return FMemberBindType(Leaf);
-}
-
-template<typename Struct, class Ids>
-FMemberBindType BindInnermostType(FOptionalInnerId& OutBindId, FMemberSpec& OutSpec)
-{
-	return BindInnerStruct(OutBindId, OutSpec, GetStructBothId<Ids, Struct>());
-}
-
-template<Arithmetic Type, class Ids>
-FMemberBindType BindInnermostType(FOptionalInnerId& OutId, FMemberSpec& OutSpec)
-{
-	return BindInnerLeaf(OutId, OutSpec, ReflectArithmetic<Type>, NoId);
-}
-
-template<Enumeration Enum, class Ids>
-FMemberBindType BindInnermostType(FOptionalInnerId& OutId, FMemberSpec& OutSpec)
-{
-	return BindInnerLeaf(OutId, OutSpec, ReflectEnum<Enum>, GetEnumId<Ids, Enum>());
-}
-
-template<typename RangeBinding>
-constexpr uint32 CountRangeBindings()
-{
-	using InnerBinding = RangeBind<typename RangeBinding::ItemType>;
-	return 1 + CountRangeBindings<InnerBinding>();
-}
-
-template<>
-constexpr uint32 CountRangeBindings<void>()
-{
-	return 0;
-}
-
-template<typename RangeBinding, uint32 NestLevel>
-struct TInnermostType
-{
-	using InnerType = typename RangeBinding::ItemType;
-	using Type = TInnermostType<RangeBind<InnerType>, NestLevel - 1>::Type;
-};
-
-template<typename RangeBinding>
-struct TInnermostType<RangeBinding, 1>
-{
-	using Type = typename RangeBinding::ItemType;
-};
-
-template<typename RangeBinding, typename Ids, uint32 N>
-TConstArrayView<FRangeBinding> GetRangeBindings() requires (N > 0)
-{
-	struct FOnce
-	{
-		FOnce(FAnsiStringView RangeBindName)
-		: Instance(Ids::IndexTypename(RangeBindName))
-		, Binding(Instance, RangeSizeOf(typename RangeBinding::SizeType{}))
-		{}
-		RangeBinding Instance;
-		FRangeBinding Binding;
-	};
-
-	if constexpr (N == 1)
-	{
-		static FOnce Static(ToAnsiView(RangeBinding::BindName));
-		return MakeArrayView(&Static.Binding, N);
-	}
-	else
-	{
-		using InnerType = typename RangeBinding::ItemType;
-		using InnerRangeBinding = RangeBind<InnerType>;
-
-		struct FNestedOnce : FOnce
-		{
-			FNestedOnce()
-			: FOnce(ToAnsiView(InnerRangeBinding::BindName))
-			{
-				FMemory::Memcpy(NestedBindings, GetRangeBindings<InnerRangeBinding, Ids,  N - 1>().GetData(), sizeof(NestedBindings));
-			}
-			
-			uint8 NestedBindings[sizeof(FRangeBinding) * (N - 1)] = {};
-		};
-		static_assert(std::is_trivially_destructible_v<FRangeBinding>);
-		static_assert(offsetof(FNestedOnce, Binding) + sizeof(FRangeBinding) == offsetof(FNestedOnce, NestedBindings));	
-
-		static FNestedOnce Static;
-		return MakeArrayView(&Static.Binding, N);
-	}
-}
-
-template<LeafType Type, class Runtime>
-FMemberBinding BindMember(uint64 Offset, FMemberSpec& OutSpec)
-{
-	FMemberBinding Out(Offset);
-	Out.InnermostType = BindInnermostType<Type, typename Runtime::Ids>(Out.InnermostId, OutSpec);
-	return Out;
-}
-
-template<typename Type, class Runtime>
-FMemberBinding BindMember(uint64 Offset, FMemberSpec& OutSpec)
-{
-	using Ids = typename Runtime::Ids;
-	using CustomBinding = typename Runtime::template CustomBindings<Type>::Type;
-
-	FMemberBinding Out(Offset);
-	if constexpr (!std::is_void_v<CustomBinding>)
-	{
-		FBothStructId Both = BindCustomStructOnce<Type, CustomBinding, Runtime>();
-		Out.InnermostType = BindInnerStruct(Out.InnermostId, OutSpec, Both);
-	}
-	else
-	{
-		using RangeBinding = RangeBind<Type>;
-		if constexpr (!std::is_void_v<RangeBinding>)
-		{
-			constexpr uint32 NumRangeBindings = CountRangeBindings<RangeBinding>();
-			using InnermostType = typename TInnermostType<RangeBinding, NumRangeBindings>::Type;
-
-			Out.RangeBindings = GetRangeBindings<RangeBinding, Ids, NumRangeBindings>();
-			Out.InnermostType = BindInnermostType<InnermostType, Ids>(Out.InnermostId, OutSpec);
-			
-			for (FRangeBinding Range : ReverseIterate(Out.RangeBindings))
-			{
-				OutSpec.RangeWrap(Range.GetSizeType());
-			}
-		}
-		else
-		{
-			FBothStructId Both = GetStructBothId<typename Runtime::Ids, Type>();
-			Out.InnermostType = BindInnerStruct(Out.InnermostId, OutSpec, Both);
-		}
-	}
-	
-	return Out;
-}
-
-////////////////////////////////////////////////////////////////////////////////////////////////
-
-template<class Ctti, class Ids>
-FEnumId DeclareNativeEnum(FEnumDeclarations& Out, EEnumMode Mode)
-{
-	using UnderlyingType = std::underlying_type_t<typename Ctti::Type>;
-
-	FType Type = IndexCttiName<Ids, Ctti>();
-	FEnumId Id = Ids::IndexEnum(Type);
-	FEnumerator Enumerators[Ctti::NumEnumerators];
-	for (FEnumerator& Enumerator : Enumerators)
-	{
-		Enumerator.Name = Ids::IndexName(Ctti::Enumerators[&Enumerator - Enumerators].Name);
-		Enumerator.Constant = static_cast<uint64>(static_cast<UnderlyingType>(Ctti::Enumerators[&Enumerator - Enumerators].Constant));
-	}
-	Out.Declare(Id, Type, Mode, Enumerators, EEnumAliases::Fail);
-
-	return Id;
-}
-
-template<class Ctti, class Runtime>
-void BindNativeStruct(FSchemaBindings& Out, FBothStructId Both, EMemberPresence Occupancy)
-{
-	using Ids = Runtime::Ids;
-	using SuperType = typename Ctti::Super;
-	FOptionalDeclId SuperId;
-	if constexpr (!std::is_void_v<SuperType>)
-	{
-		SuperId = GetStructDeclId<Ids, SuperType>();
-	}
-
-	if constexpr (Ctti::NumVars) // Temp workaround til MakeArrayView handles zero-sized arrays
-	{
-		FMemberId MemberNames[Ctti::NumVars];
-		FMemberSpec MemberTypes[Ctti::NumVars];
-		FMemberBinding MemberBindings[Ctti::NumVars];
-		ForEachVar<Ctti>([&]<class Var>()
-		{
-			MemberNames[Var::Index] = Ids::IndexMember(Var::Name);
-			FMemberSpec& OutSpec = MemberTypes[Var::Index];
-			MemberBindings[Var::Index] = BindMember<typename Var::Type, Runtime>(Var::Offset, OutSpec);
-		});
-	
-		const FStructSpec Spec = { Both.DeclId, SuperId, /* no CTTI version yet */ 0, Occupancy, MemberNames, MemberTypes };
-		Out.BindStruct(Both.BindId, MemberBindings, Spec);
-	}
-	else
-	{
-		const FStructSpec Spec = { Both.DeclId, SuperId, /* no CTTI version yet */ 0, Occupancy, {}, {} };
-		Out.BindStruct(Both.BindId, {}, Spec);
-	}
-}
-
 //////////////////////////////////////////////////////////////////////////
 
 // Save -> load struct ids for ESchemaFormat::InMemoryNames, alternative to side-channel with ExtractRuntimeIds()
@@ -1288,4 +708,4 @@ struct FIdTranslator : FIdTranslatorBase
 PLAINPROPS_API FSchemaBatch*	CreateTranslatedSchemas(const FSchemaBatch& Schemas, FIdBinding NewIds);
 PLAINPROPS_API void				DestroyTranslatedSchemas(const FSchemaBatch* Schemas);
 
-} // namespace PlainProps
+} // namespace PlainProps
\ No newline at end of file
diff --git a/Engine/Plugins/Experimental/PlainProps/Source/Public/PlainPropsBindCtti.h b/Engine/Plugins/Experimental/PlainProps/Source/Public/PlainPropsBindCtti.h
new file mode 100644
--- /dev/null
+++ b/Engine/Plugins/Experimental/PlainProps/Source/Public/PlainPropsBindCtti.h
@@ -0,0 +1,590 @@
+// Copyright Epic Games, Inc. All Rights Reserved.
+
+#pragma once
+
+#include "PlainPropsBind.h"
+#include "PlainPropsCtti.h"
+#include "PlainPropsSpecify.h"
+#include "PlainPropsTypename.h"
+#include "Containers/StringView.h"
+#include <tuple>
+
+namespace PlainProps 
+{
+
+inline FAnsiStringView ToAnsiView(std::string_view Str) { return FAnsiStringView(Str.data(), Str.length()); }
+
+////////////////////////////////////////////////////////////////////////////////////////////////
+
+template<typename T>
+struct TRangeBind{ using Type = void; };
+
+template<typename T>
+using RangeBind = typename TRangeBind<T>::Type;
+
+template<typename T>
+struct TCustomBind { using Type = void; };
+
+template<class T>
+struct TCustomDeltaBind { using Type = void; };
+
+template<typename T>
+using CustomBind = typename TCustomBind<T>::Type;
+
+template<typename T>
+using CustomDeltaBind = typename TCustomDeltaBind<T>::Type;
+
+template<typename T>
+struct TOccupancyOf
+{
+	static constexpr EMemberPresence Value = EMemberPresence::AllowSparse;
+};
+
+struct FRequireAll
+{
+	static constexpr EMemberPresence Value = EMemberPresence::RequireAll;
+};
+
+////////////////////////////////////////////////////////////////////////////////////////////////
+
+template<class Ids, class Typename>
+FScopeId IndexNamespaceId() requires (Typename::Namespace.size() > 0)
+{
+	// Opt: Make cached GetNamespaceId(), either via new namespace CTTI types (maybe PP_REFLECT_NAMESPACE)
+	//		or some compile time string template parameters, perhaps a variadic template taking any number of chars
+	return Ids::IndexScope(ToAnsiView(Typename::Namespace));
+}
+
+template<class Ids, class Typename>
+FScopeId IndexNamespaceId()
+{
+	return NoId;
+}
+
+template<ETypename Kind, class Typename>
+constexpr std::string_view SelectStructName()
+{	
+	if constexpr (Kind == ETypename::Bind && ExplicitBindName<Typename>)
+	{
+		return Typename::BindName;
+	}
+	else
+	{
+		return Typename::DeclName;	
+	}
+}
+
+template<class Ids, ETypename Kind, class Typename>
+FType IndexStructName()
+{
+	FType BaseName = { IndexNamespaceId<Ids, Typename>(), 
+						 FTypenameId(Ids::IndexTypename(ToAnsiView(SelectStructName<Kind, Typename>()))) };
+	
+	if constexpr (ParametricName<Typename>)
+	{
+		return IndexParametricType<Ids, Kind>(BaseName, (typename Typename::Parameters*)nullptr);
+	}
+	else
+	{
+		return BaseName;
+	}
+}
+
+template<class Ids, class Typename>
+FBothStructId IndexStructBothId(FType DeclName = IndexStructName<Ids, ETypename::Decl, Typename>())
+{	
+	FBothStructId Out;
+	Out.DeclId = FDeclId(Ids::IndexStruct(DeclName));
+	Out.BindId = UpCast(Out.DeclId);
+
+	if constexpr (ExplicitBindName<Typename> || ParametricName<Typename>)
+	{
+		FType BindName = IndexStructName<Ids, ETypename::Bind, Typename>();
+		Out.BindId = BindName != DeclName ? FBindId(Ids::IndexStruct(BindName)) : UpCast(Out.DeclId);
+	}
+
+	return Out;
+}
+
+template<class Ids, class Typename>
+FDualStructId IndexStructDualId(FType Name = IndexStructName<Ids, ETypename::Decl, Typename>()) requires (!ExplicitBindName<Typename>)
+{	
+	return FDualStructId(Ids::IndexStruct(Name));
+}
+
+// Cached by function static
+template<class Ids, typename Struct>
+FDeclId GetStructDeclId()
+{
+	static FDeclId Id = FDeclId(Ids::IndexStruct(IndexStructName<Ids, ETypename::Decl, TTypename<Struct>>()));
+	return Id;
+}
+
+// Cached by function static
+template<class Ids, typename Struct>
+FBindId GetStructBindId()
+{
+	static FBindId Id = FBindId(Ids::IndexStruct(IndexStructName<Ids, ETypename::Bind, TTypename<Struct>>()));
+	return Id;
+}
+
+// Cached by function static
+template<class Ids, typename Struct>
+FBothStructId GetStructBothId()
+{
+	static FBothStructId Id = IndexStructBothId<Ids, TTypename<Struct>>();
+	return Id;
+}
+
+// Cached by function static
+template<class Ids, typename Struct>
+FDualStructId GetStructDualId()
+{
+	static FDualStructId Id = IndexStructDualId<Ids, TTypename<Struct>>();
+	return Id;
+}
+
+template<class Ids, typename Ctti>
+FType IndexCttiName()
+{
+	FTypenameId Name{Ids::IndexTypename(Ctti::Name)};
+	FScopeId Namespace = NoId;
+	if constexpr (Ctti::Namespace[0] != '\0')
+	{
+		Namespace = Ids::IndexScope(ToAnsiView(Ctti::Namespace));
+	}
+	return { Namespace, Name };
+}
+
+// Cached by function static
+template<class Ids, Enumeration Enum>
+FEnumId GetEnumId()
+{
+	static FEnumId Id = Ids::IndexEnum(IndexCttiName<Ids, CttiOf<Enum>>());
+	return Id;
+}
+
+template<class Ids, Arithmetic T>
+FType IndexArithmeticName()
+{
+	static constexpr FUnpackedLeafType Leaf = ReflectArithmetic<T>;
+	return { NoId, FTypenameId(Ids::IndexTypename(ToAnsiView(ArithmeticName<Leaf.Type, Leaf.Width>))) };
+}
+
+template<class Ids, ETypename, Arithmetic Leaf>
+FType IndexParameterName()
+{
+	return IndexArithmeticName<Ids, Leaf>();
+}
+
+template<class Ids, ETypename, Enumeration Enum>
+FType IndexParameterName()
+{
+	return GetEnumId<Ids, Enum>();
+}
+
+template<class Ids, ETypename Kind, typename T>
+FType IndexParameterName()
+{
+	using RangeBinding = RangeBind<T>;
+	if constexpr (std::is_void_v<RangeBinding>)
+	{
+		return IndexStructName<Ids, Kind, TTypename<T>>();
+	}
+	else
+	{
+		FType ItemParam = IndexParameterName<Ids, Kind, typename RangeBinding::ItemType>();
+		FType SizeParam = Ids::GetIndexer().MakeRangeParameter(RangeSizeOf(typename RangeBinding::SizeType{}));
+
+		if constexpr (Kind == ETypename::Decl)
+		{
+			// Type-erase range type
+			return Ids::GetIndexer().MakeAnonymousParametricType({ItemParam, SizeParam});
+		}
+		else
+		{
+			using Typename = TTypename<T>;
+			FType RangeBindName = { IndexNamespaceId<Ids, Typename>(), FTypenameId(Ids::IndexTypename(ToAnsiView(RangeBinding::BindName))) };
+			return Ids::GetIndexer().MakeParametricType(RangeBindName, {ItemParam, SizeParam});
+		}
+	}
+}
+
+template<class Ids, ETypename Kind, typename... Ts>
+FType IndexParametricType(FType TemplatedType, const std::tuple<Ts...>*)
+{
+	FType Parameters[] = { (IndexParameterName<Ids, Kind, Ts>())... };
+	return Ids::GetIndexer().MakeParametricType(TemplatedType, Parameters);
+}
+
+////////////////////////////////////////////////////////////////////////////////////////////////
+
+// Works around C++'s lack of templated nullary constructors
+template<class Ids>
+struct TInit {};
+
+using InnerStructArray = TArray<FInnerStruct, TInlineAllocator<8>>;
+
+struct FCustomInit
+{
+	UE_NONCOPYABLE(FCustomInit);
+	FCustomInit() = default;
+
+	InnerStructArray				LoweredInners;
+
+	// Only needed for inner structs that might get type-erased/lowered
+	void RegisterInnerStruct(FBothStructId InnerId, TConstArrayView<FMemberId> Names)
+	{
+		if (InnerId.IsLowered())
+		{
+			for (FMemberId Name : Names)
+			{
+				LoweredInners.Add({Name, InnerId.BindId});
+			}
+		}
+	}
+};
+
+// Helps custom binding constructors create ids and register types that might get type-erased 
+template<class Ids>
+struct TCustomInit : FCustomInit, TInit<Ids>
+{};
+
+// Constructor API for static custom bindings to create ids, register type-erased inner structs
+// and specify member types, inheritance and version 
+template<class Ids, uint32 N>
+struct TCustomSpecifier : TCustomInit<Ids>
+{
+	FOptionalStructId				Super;
+	uint16							Version = 0;
+	FMemberSpec						Members[N];
+
+	template <typename... Ts>
+    void SetMembers(Ts&&... Specs) requires (N == sizeof...(Specs))
+    {
+		FMemberSpec* It = &Members[0];
+		((*It++ = FMemberSpec(Forward<Ts>(Specs))), ...);
+    }
+
+	void FillMembers(FMemberSpec Same)
+	{
+		for (FMemberSpec& Member : Members)
+		{
+			Member = Same;
+		}
+	}
+};
+
+////////////////////////////////////////////////////////////////////////////////////////////////
+
+// Workaround for clang requiring the definition, make it dependent on any other type
+template<typename>
+struct TDependent
+{
+	using Baseline = FBaseline;
+	using StructLoadView = FStructLoadView;
+};
+
+// Scoped custom binding for native type
+template<class T, class Runtime, typename CustomBinding = CustomBind<T>>
+struct TScopedStructBinding : FBothStructId, CustomBinding
+{
+	using Ids = typename Runtime::Ids;
+	using Typename = TCustomTypename<CustomBinding>::Type;
+	static constexpr uint32 N = sizeof(CustomBinding::MemberIds) / sizeof(CustomBinding::MemberIds[0]);
+	static constexpr bool bFastLoad = requires(FMemcpyLoadPlan& O) { CustomBinding::Plan(O); };
+
+	explicit TScopedStructBinding(FType DeclName = IndexStructName<Ids, ETypename::Decl, Typename>())
+	: TScopedStructBinding(IndexStructBothId<Ids, Typename>(DeclName))
+	{}
+
+	explicit TScopedStructBinding(FDualStructId Dual)
+	: TScopedStructBinding(FBothStructId{Dual, Dual})
+	{}
+
+	explicit TScopedStructBinding(FBothStructId Both)
+	: TScopedStructBinding(Both, TCustomSpecifier<Ids, N>())
+	{}
+
+	TScopedStructBinding(FBothStructId Both, TCustomSpecifier<Ids, N>&& Init)
+	: FBothStructId(Both)
+	, CustomBinding(/* out */ Init)
+	{
+		FStructSpec Decl = { DeclId, /* todo */ NoId, Init.Version, TOccupancyOf<T>::Value, CustomBinding::MemberIds, Init.Members };
+		Runtime::GetCustoms().BindStruct(BindId, *this, Decl, Init.LoweredInners);
+	}
+
+	~TScopedStructBinding()
+	{
+		Runtime::GetCustoms().DropStruct(BindId);
+	}
+
+	virtual void SaveCustom(FMemberBuilder& Dst, const void* Src, TDependent<Ids>::Baseline Base, const FSaveContext& Ctx) override
+	{
+		const void* Default = Base.Get();
+		CustomBinding::Save(Dst, *static_cast<const T*>(Src), static_cast<const T*>(Default), Ctx);
+	}
+
+	virtual void LoadCustom(void* Dst, TDependent<Ids>::StructLoadView Src, ECustomLoadMethod Method) const override 
+	{
+		if constexpr (bFastLoad)
+		{
+			unimplemented();
+		}
+		else
+		{
+			CustomBinding::Load(*static_cast<T*>(Dst), Src, Method);
+		}
+	}
+	
+	virtual bool DiffCustom(const void* A, const void* B, const FBindContext& Ctx) const override
+	{
+		return CustomBinding::Diff(*static_cast<const T*>(A), *static_cast<const T*>(B), Ctx);
+	}
+
+	virtual bool DiffCustom(const void* A, const void* B, FDiffContext& Ctx) const override
+	{
+		return CustomBinding::Diff(*static_cast<const T*>(A), *static_cast<const T*>(B), Ctx);
+	}
+
+	virtual void PlanCustom(FMemcpyLoadPlan& Out) const override
+	{
+		if constexpr (bFastLoad)
+		{
+			CustomBinding::Plan(Out);
+		}
+	}
+};
+
+// Specialization for schemabound structs
+template<class T, class Runtime>
+struct TScopedStructBinding<T, Runtime, void> : FBothStructId
+{
+	using Ids = typename Runtime::Ids;
+	using Typename = TTypename<T>;
+	
+	TScopedStructBinding(FBothStructId Both = IndexStructBothId<Ids, Typename>() )
+	: FBothStructId(Both)
+	{
+		BindNativeStruct<CttiOf<T>, Runtime>(/* out */ Runtime::GetSchemas(), Both, TOccupancyOf<T>::Value);
+	}
+
+	~TScopedStructBinding()
+	{
+		Runtime::GetSchemas().DropStruct(BindId);
+	}
+};
+
+template<class Struct, typename CustomBinding, class Runtime>
+FBothStructId BindCustomStructOnce()
+{
+	static TScopedStructBinding<Struct, Runtime, CustomBinding> Instance;
+	return Instance;
+}
+
+////////////////////////////////////////////////////////////////////////////////////////////////
+
+inline constexpr FMemberBindType DefaultStructBindType = FMemberBindType(FStructType{EMemberKind::Struct, /* IsDynamic */ 0, /* IsSuper */ 0});
+inline constexpr FMemberBindType SuperStructBindType = FMemberBindType(FStructType{EMemberKind::Struct, /* IsDynamic */ 0, /* IsSuper */ 1});
+
+inline FMemberBindType BindInnerStruct(FOptionalInnerId& OutBindId, FMemberSpec& OutSpec, FBothStructId Both)
+{
+	OutBindId = FInnerId(Both.BindId);
+	OutSpec = FMemberSpec(Both.DeclId);
+	return DefaultStructBindType;
+}
+
+inline FMemberBindType BindInnerLeaf(FOptionalInnerId& OutId, FMemberSpec& OutSpec, FUnpackedLeafType Leaf, FOptionalEnumId Id)
+{
+	OutId = ToOptionalInner(Id);
+	OutSpec = FMemberSpec(Leaf, Id);
+	return FMemberBindType(Leaf);
+}
+
+template<typename Struct, class Ids>
+FMemberBindType BindInnermostType(FOptionalInnerId& OutBindId, FMemberSpec& OutSpec)
+{
+	return BindInnerStruct(OutBindId, OutSpec, GetStructBothId<Ids, Struct>());
+}
+
+template<Arithmetic Type, class Ids>
+FMemberBindType BindInnermostType(FOptionalInnerId& OutId, FMemberSpec& OutSpec)
+{
+	return BindInnerLeaf(OutId, OutSpec, ReflectArithmetic<Type>, NoId);
+}
+
+template<Enumeration Enum, class Ids>
+FMemberBindType BindInnermostType(FOptionalInnerId& OutId, FMemberSpec& OutSpec)
+{
+	return BindInnerLeaf(OutId, OutSpec, ReflectEnum<Enum>, GetEnumId<Ids, Enum>());
+}
+
+template<typename RangeBinding>
+constexpr uint32 CountRangeBindings()
+{
+	using InnerBinding = RangeBind<typename RangeBinding::ItemType>;
+	return 1 + CountRangeBindings<InnerBinding>();
+}
+
+template<>
+constexpr uint32 CountRangeBindings<void>()
+{
+	return 0;
+}
+
+template<typename RangeBinding, uint32 NestLevel>
+struct TInnermostType
+{
+	using InnerType = typename RangeBinding::ItemType;
+	using Type = TInnermostType<RangeBind<InnerType>, NestLevel - 1>::Type;
+};
+
+template<typename RangeBinding>
+struct TInnermostType<RangeBinding, 1>
+{
+	using Type = typename RangeBinding::ItemType;
+};
+
+template<typename RangeBinding, typename Ids, uint32 N>
+TConstArrayView<FRangeBinding> GetRangeBindings() requires (N > 0)
+{
+	struct FOnce
+	{
+		FOnce(FAnsiStringView RangeBindName)
+		: Instance(Ids::IndexTypename(RangeBindName))
+		, Binding(Instance, RangeSizeOf(typename RangeBinding::SizeType{}))
+		{}
+		RangeBinding Instance;
+		FRangeBinding Binding;
+	};
+
+	if constexpr (N == 1)
+	{
+		static FOnce Static(ToAnsiView(RangeBinding::BindName));
+		return MakeArrayView(&Static.Binding, N);
+	}
+	else
+	{
+		using InnerType = typename RangeBinding::ItemType;
+		using InnerRangeBinding = RangeBind<InnerType>;
+
+		struct FNestedOnce : FOnce
+		{
+			FNestedOnce()
+			: FOnce(ToAnsiView(InnerRangeBinding::BindName))
+			{
+				FMemory::Memcpy(NestedBindings, GetRangeBindings<InnerRangeBinding, Ids,  N - 1>().GetData(), sizeof(NestedBindings));
+			}
+			
+			uint8 NestedBindings[sizeof(FRangeBinding) * (N - 1)] = {};
+		};
+		static_assert(std::is_trivially_destructible_v<FRangeBinding>);
+		static_assert(offsetof(FNestedOnce, Binding) + sizeof(FRangeBinding) == offsetof(FNestedOnce, NestedBindings));	
+
+		static FNestedOnce Static;
+		return MakeArrayView(&Static.Binding, N);
+	}
+}
+
+template<LeafType Type, class Runtime>
+FMemberBinding BindMember(uint64 Offset, FMemberSpec& OutSpec)
+{
+	FMemberBinding Out(Offset);
+	Out.InnermostType = BindInnermostType<Type, typename Runtime::Ids>(Out.InnermostId, OutSpec);
+	return Out;
+}
+
+template<typename Type, class Runtime>
+FMemberBinding BindMember(uint64 Offset, FMemberSpec& OutSpec)
+{
+	using Ids = typename Runtime::Ids;
+	using CustomBinding = typename Runtime::template CustomBindings<Type>::Type;
+
+	FMemberBinding Out(Offset);
+	if constexpr (!std::is_void_v<CustomBinding>)
+	{
+		FBothStructId Both = BindCustomStructOnce<Type, CustomBinding, Runtime>();
+		Out.InnermostType = BindInnerStruct(Out.InnermostId, OutSpec, Both);
+	}
+	else
+	{
+		using RangeBinding = RangeBind<Type>;
+		if constexpr (!std::is_void_v<RangeBinding>)
+		{
+			constexpr uint32 NumRangeBindings = CountRangeBindings<RangeBinding>();
+			using InnermostType = typename TInnermostType<RangeBinding, NumRangeBindings>::Type;
+
+			Out.RangeBindings = GetRangeBindings<RangeBinding, Ids, NumRangeBindings>();
+			Out.InnermostType = BindInnermostType<InnermostType, Ids>(Out.InnermostId, OutSpec);
+			
+			for (FRangeBinding Range : ReverseIterate(Out.RangeBindings))
+			{
+				OutSpec.RangeWrap(Range.GetSizeType());
+			}
+		}
+		else
+		{
+			FBothStructId Both = GetStructBothId<typename Runtime::Ids, Type>();
+			Out.InnermostType = BindInnerStruct(Out.InnermostId, OutSpec, Both);
+		}
+	}
+	
+	return Out;
+}
+
+////////////////////////////////////////////////////////////////////////////////////////////////
+
+template<class Ctti, class Ids>
+FEnumId DeclareNativeEnum(FEnumDeclarations& Out, EEnumMode Mode)
+{
+	using UnderlyingType = std::underlying_type_t<typename Ctti::Type>;
+
+	FType Type = IndexCttiName<Ids, Ctti>();
+	FEnumId Id = Ids::IndexEnum(Type);
+	FEnumerator Enumerators[Ctti::NumEnumerators];
+	for (FEnumerator& Enumerator : Enumerators)
+	{
+		Enumerator.Name = Ids::IndexName(Ctti::Enumerators[&Enumerator - Enumerators].Name);
+		Enumerator.Constant = static_cast<uint64>(static_cast<UnderlyingType>(Ctti::Enumerators[&Enumerator - Enumerators].Constant));
+	}
+	Out.Declare(Id, Type, Mode, Enumerators, EEnumAliases::Fail);
+
+	return Id;
+}
+
+////////////////////////////////////////////////////////////////////////////////////////////////
+
+template<class Ctti, class Runtime>
+void BindNativeStruct(FSchemaBindings& Out, FBothStructId Both, EMemberPresence Occupancy)
+{
+	using Ids = Runtime::Ids;
+	using SuperType = typename Ctti::Super;
+	FOptionalDeclId SuperId;
+	if constexpr (!std::is_void_v<SuperType>)
+	{
+		SuperId = GetStructDeclId<Ids, SuperType>();
+	}
+
+	if constexpr (Ctti::NumVars) // Temp workaround til MakeArrayView handles zero-sized arrays
+	{
+		FMemberId MemberNames[Ctti::NumVars];
+		FMemberSpec MemberTypes[Ctti::NumVars];
+		FMemberBinding MemberBindings[Ctti::NumVars];
+		ForEachVar<Ctti>([&]<class Var>()
+		{
+			MemberNames[Var::Index] = Ids::IndexMember(Var::Name);
+			FMemberSpec& OutSpec = MemberTypes[Var::Index];
+			MemberBindings[Var::Index] = BindMember<typename Var::Type, Runtime>(Var::Offset, OutSpec);
+		});
+	
+		const FStructSpec Spec = { Both.DeclId, SuperId, /* no CTTI version yet */ 0, Occupancy, MemberNames, MemberTypes };
+		Out.BindStruct(Both.BindId, MemberBindings, Spec);
+	}
+	else
+	{
+		const FStructSpec Spec = { Both.DeclId, SuperId, /* no CTTI version yet */ 0, Occupancy, {}, {} };
+		Out.BindStruct(Both.BindId, {}, Spec);
+	}
+}
+
+} // namespace PlainProps
\ No newline at end of file
diff --git a/Engine/Plugins/Experimental/PlainProps/Source/Public/PlainPropsSaveMember.h b/Engine/Plugins/Experimental/PlainProps/Source/Public/PlainPropsSaveMember.h
--- a/Engine/Plugins/Experimental/PlainProps/Source/Public/PlainPropsSaveMember.h
+++ b/Engine/Plugins/Experimental/PlainProps/Source/Public/PlainPropsSaveMember.h
@@ -3,7 +3,7 @@
 #pragma once
 
 #include "PlainPropsTypes.h"
-#include "PlainPropsBind.h"
+#include "PlainPropsBindCtti.h"
 #include "PlainPropsBuild.h"
 #include "PlainPropsDiff.h"
 #include "PlainPropsLoadMember.h"
diff --git a/Engine/Plugins/Experimental/PlainProps/Source/Public/PlainPropsUeCoreBindings.h b/Engine/Plugins/Experimental/PlainProps/Source/Public/PlainPropsUeCoreBindings.h
--- a/Engine/Plugins/Experimental/PlainProps/Source/Public/PlainPropsUeCoreBindings.h
+++ b/Engine/Plugins/Experimental/PlainProps/Source/Public/PlainPropsUeCoreBindings.h
@@ -10,7 +10,7 @@
 #include "Math/Transform.h"
 #include "Misc/Optional.h"
 #include "Misc/TVariant.h"
-#include "PlainPropsBind.h"
+#include "PlainPropsBindCtti.h"
 #include "PlainPropsIndex.h"
 #include "PlainPropsLoadMember.h"
 #include "PlainPropsSaveMember.h"
diff --git a/Engine/Plugins/Experimental/PlainPropsUObject/Source/Public/PlainPropsUObjectRuntime.h b/Engine/Plugins/Experimental/PlainPropsUObject/Source/Public/PlainPropsUObjectRuntime.h
--- a/Engine/Plugins/Experimental/PlainPropsUObject/Source/Public/PlainPropsUObjectRuntime.h
+++ b/Engine/Plugins/Experimental/PlainPropsUObject/Source/Public/PlainPropsUObjectRuntime.h
@@ -2,7 +2,7 @@
 
 #pragma once
 
-#include "PlainPropsBind.h"
+#include "PlainPropsBindCtti.h"
 #include "PlainPropsDeclare.h"
 #include "PlainPropsIndex.h"
 #include "PlainPropsTypes.h"
