PlainProps unimplemented Aux binding API draft

#rnx

diff --git a/Engine/Plugins/Experimental/PlainProps/Source/Public/PlainPropsBindCtti.h b/Engine/Plugins/Experimental/PlainProps/Source/Public/PlainPropsBindCtti.h
--- a/Engine/Plugins/Experimental/PlainProps/Source/Public/PlainPropsBindCtti.h
+++ b/Engine/Plugins/Experimental/PlainProps/Source/Public/PlainPropsBindCtti.h
@@ -6,6 +6,7 @@
 #include "PlainPropsCtti.h"
 #include "PlainPropsSpecify.h"
 #include "PlainPropsTypename.h"
+#include "Containers/StaticArray.h"
 #include "Containers/StringView.h"
 #include <tuple>
 
@@ -587,4 +588,28 @@ void BindNativeStruct(FSchemaBindings& Out, FBothStructId Both, EMemberPresence
 	}
 }
 
+////////////////////////////////////////////////////////////////////////////////////////////////
+
+struct FAuxMember
+{
+	FMemberId			Name;
+	FMemberBinding		Binding;
+	FMemberSpec			Spec;
+};
+
+template<class Ctti, class Runtime>
+TStaticArray<FAuxMember, Ctti::NumVars> BindAuxMembers()
+{
+	static_assert(std::is_void_v<typename Ctti::Super>, "Inheritance support unimplemented");
+
+	TStaticArray<FAuxMember, Ctti::NumVars> Out;
+	ForEachVar<Ctti>([&]<class Var>()
+	{
+		FAuxMember& Aux = Out[Var::Index];
+		Aux.Name = Runtime::Ids::IndexMember(Var::Name);
+		Aux.Binding = BindMember<typename Var::Type, Runtime>(Var::Offset, /* out */ Aux.Spec);
+	});
+	return Out;
+}
+
 } // namespace PlainProps
\ No newline at end of file
diff --git a/Engine/Plugins/Experimental/PlainPropsUObject/Source/Private/PlainPropsUObjectBindings.cpp b/Engine/Plugins/Experimental/PlainPropsUObject/Source/Private/PlainPropsUObjectBindings.cpp
--- a/Engine/Plugins/Experimental/PlainPropsUObject/Source/Private/PlainPropsUObjectBindings.cpp
+++ b/Engine/Plugins/Experimental/PlainPropsUObject/Source/Private/PlainPropsUObjectBindings.cpp
@@ -3068,6 +3068,7 @@ public:
 
 	void											BindSuper(FDeclId SuperId);
 	void											BindMember(FProperty* Property);
+	void											Interlace(TConstArrayView<FAuxMember> Members);
 	TConstArrayView<FMemberBinding>					GetMembers() const { return Members; }
 	FStructDeclarationPtr							Declare() const;	
 
