Fixed a missing include, which causes a non-unity compile fail

diff --git a/Engine/Plugins/Experimental/PlainProps/Source/Private/PlainPropsParse.cpp b/Engine/Plugins/Experimental/PlainProps/Source/Private/PlainPropsParse.cpp
--- a/Engine/Plugins/Experimental/PlainProps/Source/Private/PlainPropsParse.cpp
+++ b/Engine/Plugins/Experimental/PlainProps/Source/Private/PlainPropsParse.cpp
@@ -7,6 +7,7 @@
 #include "Misc/StringBuilder.h"
 #include "PlainPropsBind.h"
 #include "PlainPropsBuild.h"
+#include "PlainPropsSpecify.h"
 #include "PlainPropsBuildSchema.h"
 #include "PlainPropsInternalBuild.h"
 #include "PlainPropsInternalParse.h"
