CIS Fix: Add missing include for non unity builds.

#rnx

diff --git a/Engine/Plugins/Experimental/PlainProps/Source/Tests/PlainPropsTests.cpp b/Engine/Plugins/Experimental/PlainProps/Source/Tests/PlainPropsTests.cpp
--- a/Engine/Plugins/Experimental/PlainProps/Source/Tests/PlainPropsTests.cpp
+++ b/Engine/Plugins/Experimental/PlainProps/Source/Tests/PlainPropsTests.cpp
@@ -19,6 +19,7 @@
 #include "PlainPropsInternalPrint.h"
 #include "PlainPropsInternalRead.h"
 #include "PlainPropsRead.h"
+#include "PlainPropsSpecify.h"
 #include "PlainPropsVisualize.h"
 #include "PlainPropsWrite.h"
 #include "Templates/UnrealTemplate.h"
