PlainProps add optional enum width declaration and checked declaration lookup

#rnx

diff --git a/Engine/Plugins/Experimental/PlainProps/Source/Private/PlainPropsDeclare.cpp b/Engine/Plugins/Experimental/PlainProps/Source/Private/PlainPropsDeclare.cpp
--- a/Engine/Plugins/Experimental/PlainProps/Source/Private/PlainPropsDeclare.cpp
+++ b/Engine/Plugins/Experimental/PlainProps/Source/Private/PlainPropsDeclare.cpp
@@ -99,7 +99,7 @@ void CopyItems(T* It, TConstArrayView<T> Items)
 	}
 }
 
-const FEnumDeclaration& FEnumDeclarations::Declare(FEnumId Id, FType Type, EEnumMode Mode, TConstArrayView<FEnumerator> Enumerators, EEnumAliases Policy)
+const FEnumDeclaration& FEnumDeclarations::Declare(FEnumId Id, FType Type, EEnumMode Mode, FOptionalLeafWidth Width, TConstArrayView<FEnumerator> Enumerators, EEnumAliases Policy)
 {
 	if (static_cast<int32>(Id.Idx) >= Declarations.Num())
 	{
@@ -115,7 +115,7 @@ const FEnumDeclaration& FEnumDeclarations::Declare(FEnumId Id, FType Type, EEnum
 		Enumerators = StripAliases(/* out */ Tmp, Enumerators, Mode, Debug);
 	}
 
-	FEnumDeclaration Header{Type, Mode, IntCastChecked<uint16>(Enumerators.Num())};
+	FEnumDeclaration Header{Type, Mode, Width, IntCastChecked<uint16>(Enumerators.Num())};
 	void* Data = FMemory::Malloc(sizeof(FEnumDeclaration) + Enumerators.Num() * Enumerators.GetTypeSize());
 	Ptr.Reset(new (Data) FEnumDeclaration(Header));
 	CopyItems(Ptr->Enumerators, Enumerators);
@@ -125,6 +125,20 @@ const FEnumDeclaration& FEnumDeclarations::Declare(FEnumId Id, FType Type, EEnum
 	return *Ptr;
 }
 
+void FEnumDeclarations::Specify(FEnumId Id, ELeafWidth In)
+{
+	Check(Id);
+	FOptionalLeafWidth& Width = Declarations[Id.Idx]->Width;
+	if (Width)
+	{
+		check(Width.Get() == In);
+	}
+	else
+	{
+		Width = In;
+	}
+}
+
 #if DO_CHECK
 void FEnumDeclarations::Check(FEnumId Id) const
 {
@@ -235,4 +249,20 @@ FMemberSpec::FMemberSpec(FMemberType Type, TConstArrayView<FMemberType> InnerRan
 	}
 }
 
+////////////////////////////////////////////////////////////////////////////////////////////////
+
+const FEnumDeclaration&	IDeclarations::Get(FEnumId Id) const
+{
+	const FEnumDeclaration* Out = Find(Id);
+	check(Out);
+	return *Out;
+}
+
+const FStructDeclaration& IDeclarations::Get(FStructId Id) const
+{
+	const FStructDeclaration* Out = Find(Id);
+	check(Out);
+	return *Out;
+}
+
 } // namespace PlainProps
\ No newline at end of file
diff --git a/Engine/Plugins/Experimental/PlainProps/Source/Private/PlainPropsParse.cpp b/Engine/Plugins/Experimental/PlainProps/Source/Private/PlainPropsParse.cpp
--- a/Engine/Plugins/Experimental/PlainProps/Source/Private/PlainPropsParse.cpp
+++ b/Engine/Plugins/Experimental/PlainProps/Source/Private/PlainPropsParse.cpp
@@ -1089,7 +1089,7 @@ FParsedEnumSchema FEnumSchemaParser::DeclareEnumSchema(
 		Enumerators.Emplace(Names.MakeName(ConstantNames[I]), ConstantValues[I]);
 	}
 
-	const FEnumDeclaration& Decl = Declarations.Declare(Id, Type, Mode, Enumerators, EEnumAliases::Strip);
+	const FEnumDeclaration& Decl = Declarations.Declare(Id, Type, Mode, Width, Enumerators, EEnumAliases::Strip);
 
 	return { Id, Width, Decl.GetEnumerators() };
 }
diff --git a/Engine/Plugins/Experimental/PlainProps/Source/Public/PlainPropsBindCtti.h b/Engine/Plugins/Experimental/PlainProps/Source/Public/PlainPropsBindCtti.h
--- a/Engine/Plugins/Experimental/PlainProps/Source/Public/PlainPropsBindCtti.h
+++ b/Engine/Plugins/Experimental/PlainProps/Source/Public/PlainPropsBindCtti.h
@@ -385,7 +385,7 @@ struct TScopedEnumDeclaration
 	using Ids = typename Runtime::Ids;
 
 	FEnumId Id;
-	TScopedEnumDeclaration() : Id(DeclareNativeEnum<Enum, Ids>(Runtime::GetEnums(), Mode)) {}
+	TScopedEnumDeclaration() : Id(DeclareNativeEnum<Enum, Ids>(Runtime::GetEnums(), Mode, LeafWidth<sizeof(Enum)>)) {}
 	~TScopedEnumDeclaration() { Runtime::GetEnums().Drop(Id); }
 };
 
@@ -549,7 +549,7 @@ FMemberBinding BindMember(uint64 Offset, FMemberSpec& OutSpec)
 ////////////////////////////////////////////////////////////////////////////////////////////////
 
 template<Enumeration Enum, class Ids>
-FEnumId DeclareNativeEnum(FEnumDeclarations& Out, EEnumMode Mode)
+FEnumId DeclareNativeEnum(FEnumDeclarations& Out, EEnumMode Mode, ELeafWidth Width)
 {
 	using Ctti = CttiOf<Enum>;
 	using UnderlyingType = std::underlying_type_t<typename Ctti::Type>;
@@ -562,7 +562,7 @@ FEnumId DeclareNativeEnum(FEnumDeclarations& Out, EEnumMode Mode)
 		Enumerator.Name = Ids::IndexName(Ctti::Enumerators[&Enumerator - Enumerators].Name);
 		Enumerator.Constant = static_cast<uint64>(static_cast<UnderlyingType>(Ctti::Enumerators[&Enumerator - Enumerators].Constant));
 	}
-	Out.Declare(Id, Type, Mode, Enumerators, EEnumAliases::Fail);
+	Out.Declare(Id, Type, Mode, Width, Enumerators, EEnumAliases::Fail);
 
 	return Id;
 }
diff --git a/Engine/Plugins/Experimental/PlainProps/Source/Public/PlainPropsDeclare.h b/Engine/Plugins/Experimental/PlainProps/Source/Public/PlainPropsDeclare.h
--- a/Engine/Plugins/Experimental/PlainProps/Source/Public/PlainPropsDeclare.h
+++ b/Engine/Plugins/Experimental/PlainProps/Source/Public/PlainPropsDeclare.h
@@ -14,26 +14,36 @@ namespace PlainProps
 {
 struct FStructSpec;
 
+class FOptionalLeafWidth
+{
+	uint8 Value = 0xFF;
+public:
+	FOptionalLeafWidth() = default;
+	FOptionalLeafWidth(ELeafWidth Width) : Value(static_cast<uint8>(Width)) {}
+	//void							Set(ELeafWidth Width)		{ Value = static_cast<uint8>(Width);}
+	ELeafWidth						Get() const					{ check(*this); return static_cast<ELeafWidth>(Value); }
+	explicit operator				bool() const				{ return Value != 0xFF; }
+};
+
+////////////////////////////////////////////////////////////////////////////////////////////////
+
 struct FEnumerator
 {
-	FNameId					Name;
-	uint64					Constant;
+	FNameId							Name;
+	uint64							Constant;
 };
 
-enum class EEnumMode { Flat, Flag };
+enum class EEnumMode : uint8 { Flat, Flag };
 
 struct FEnumDeclaration
 {
-	FType					Type;			// Could be removed
-	EEnumMode				Mode;
-	//EOptionalLeafWidth	Width;			// Possible approach to declare enums w/o any noted values
-	uint16					NumEnumerators;
-	FEnumerator				Enumerators[0];	// Constants must be unique, no aliases allowed
-
-	TConstArrayView<FEnumerator> GetEnumerators() const
-	{
-		return MakeArrayView(Enumerators, NumEnumerators);
-	}
+	FType							Type;			// Could be removed
+	EEnumMode						Mode;
+	FOptionalLeafWidth				Width;
+	uint16							NumEnumerators;
+	FEnumerator						Enumerators[0];	// Constants must be unique, no aliases allowed
+
+	TConstArrayView<FEnumerator>	GetEnumerators() const		{ return MakeArrayView(Enumerators, NumEnumerators);	}
 };
 
 ////////////////////////////////////////////////////////////////////////////////////////////////
@@ -48,7 +58,8 @@ public:
 	UE_NONCOPYABLE(FEnumDeclarations);
 	explicit FEnumDeclarations(FDebugIds In) : Debug(In) {}
 
-	PLAINPROPS_API const FEnumDeclaration&	Declare(FEnumId Id, FType Type, EEnumMode Mode, TConstArrayView<FEnumerator> Enumerators, EEnumAliases Policy);
+	PLAINPROPS_API const FEnumDeclaration&	Declare(FEnumId Id, FType Type, EEnumMode Mode, FOptionalLeafWidth Width, TConstArrayView<FEnumerator> Enumerators, EEnumAliases Policy);
+	PLAINPROPS_API void						Specify(FEnumId Id, ELeafWidth Width);
 	void									Drop(FEnumId Id)				{ Check(Id); Declarations[Id.Idx].Reset(); }
 	const FEnumDeclaration&					Get(FEnumId Id) const			{ Check(Id); return *Declarations[Id.Idx].Get(); }
 	const FEnumDeclaration*					Find(FEnumId Id) const			{ return Id.Idx < (uint32)Declarations.Num() ? Declarations[Id.Idx].Get() : nullptr; }
@@ -104,9 +115,12 @@ PLAINPROPS_API FStructDeclarationPtr Declare(FStructSpec Spec);
 
 struct IDeclarations
 {
-	virtual FDeclId						Lower(FBindId BindId) const = 0;
-	virtual const FEnumDeclaration*		Find(FEnumId Id) const = 0;
-	virtual const FStructDeclaration*	Find(FStructId Id) const = 0;
+	virtual FDeclId								Lower(FBindId BindId) const = 0;
+	virtual const FEnumDeclaration*				Find(FEnumId Id) const = 0;
+	virtual const FStructDeclaration*			Find(FStructId Id) const = 0;
+	
+	PLAINPROPS_API const FEnumDeclaration&		Get(FEnumId Id) const;
+	PLAINPROPS_API const FStructDeclaration&	Get(FStructId Id) const;
 };
 
 } // namespace PlainProps
diff --git a/Engine/Plugins/Experimental/PlainProps/Source/Tests/PlainPropsTests.cpp b/Engine/Plugins/Experimental/PlainProps/Source/Tests/PlainPropsTests.cpp
--- a/Engine/Plugins/Experimental/PlainProps/Source/Tests/PlainPropsTests.cpp
+++ b/Engine/Plugins/Experimental/PlainProps/Source/Tests/PlainPropsTests.cpp
@@ -196,10 +196,10 @@ FDeclId FTestBatchBuilder::DeclareStruct(const char* Scope, const char* Name, st
 	return DeclareStruct(MakeType(Scope, Name), MemberNames, MemberTypes, Occupancy, Super);
 }
 
-FEnumId FTestBatchBuilder::DeclareEnum(FType Type, EEnumMode Mode, ELeafWidth, std::initializer_list<const char*> InNames, std::initializer_list<uint64> Constants)
+FEnumId FTestBatchBuilder::DeclareEnum(FType Type, EEnumMode Mode, ELeafWidth Width, std::initializer_list<const char*> InNames, std::initializer_list<uint64> Constants)
 {
 	FEnumId Id = IndexEnum(Type);
-	EnumDeclarations.Declare(Id, Type, Mode, MakeEnumerators(InNames, Constants), EEnumAliases::Fail);
+	EnumDeclarations.Declare(Id, Type, Mode, Width, MakeEnumerators(InNames, Constants), EEnumAliases::Fail);
 	return Id;
 }
 
diff --git a/Engine/Plugins/Experimental/PlainPropsUObject/Source/Private/PlainPropsUObjectBindings.cpp b/Engine/Plugins/Experimental/PlainPropsUObject/Source/Private/PlainPropsUObjectBindings.cpp
--- a/Engine/Plugins/Experimental/PlainPropsUObject/Source/Private/PlainPropsUObjectBindings.cpp
+++ b/Engine/Plugins/Experimental/PlainPropsUObject/Source/Private/PlainPropsUObjectBindings.cpp
@@ -970,7 +970,7 @@ static FEnumId DeclareEnum(UEnum* Enum)
 	checkSlow(Enumerators.Num() == Num || Enum->ContainsExistingMax() || Enum->HasMetaData(TEXT("Hidden"), Num));
 #endif
 
-	GUE.Enums.Declare(Id, Type, Mode, Enumerators, EEnumAliases::Strip);
+	GUE.Enums.Declare(Id, Type, Mode, /* size unknown */ {}, Enumerators, EEnumAliases::Strip);
 	return Id;
 }
 
@@ -3185,9 +3185,11 @@ private:
 	inline FMemberBinding BindEnum(FEnumProperty* Property)
 	{
 		FMemberBinding Out(Property->GetOffset_ForInternal());
-		Out.InnermostId = FInnerId(GUE.Names.IndexEnum(IndexType(Property->GetEnum())));
+		FEnumId Id = GUE.Names.IndexEnum(IndexType(Property->GetEnum()));
+		Out.InnermostId = FInnerId(Id);
 		FUnpackedLeafType Leaf = {ELeafType::Enum, WidthOf(Property->GetElementSize())};
 		Out.InnermostType = FMemberBindType(Leaf);
+		GUE.Enums.Specify(Id, Leaf.Width); // UEnum doesn't know width/sizeof
 		return Out;
 	}
 
