PlainProps: Move free functions into FLoadPlanner member functions - part 1 - a straight copy and paste.

#rnx

diff --git a/Engine/Plugins/Experimental/PlainProps/Source/Private/PlainPropsLoad.cpp b/Engine/Plugins/Experimental/PlainProps/Source/Private/PlainPropsLoad.cpp
--- a/Engine/Plugins/Experimental/PlainProps/Source/Private/PlainPropsLoad.cpp
+++ b/Engine/Plugins/Experimental/PlainProps/Source/Private/PlainPropsLoad.cpp
@@ -288,214 +288,6 @@ struct FMemberLoadBinder : FMemberBinderBase
 };
 
 
-static void CopyMemberBinding(FLeafMemberBinding Binding, /* in-out */ const FSchemaId*& InnerIdIt, FMemberLoadBinder& Out)
-{
-	InnerIdIt += (Binding.Leaf.Type == ELeafBindType::Enum); // Skip enum schema
-	Out.AddMember(Binding.Leaf.Pack(), static_cast<uint32>(Binding.Offset));
-}
-
-// TODO: Make all below free functions FLoadPlanner member functions and make LoadIds a member variable
-
-template<class LoadIdBinder>
-void CopyMemberBinding(FStructMemberBinding Binding, /* in-out */ const FSchemaId*& InnerIdIt, LoadIdBinder& LoadIds, FMemberLoadBinder& Out)
-{
-	check(!Binding.Type.IsSuper);
-	check(!Binding.Type.IsDynamic); // TODO: Build out dynamic struct support
-	FStructSchemaId LoadId = LoadIds.BindLoadId(static_cast<FStructSchemaId>(*InnerIdIt), Binding.Id);
-	Out.AddMember(FMemberBindType(Binding.Type), static_cast<uint32>(Binding.Offset));
-	Out.AddInnerId(LoadId);
-	++InnerIdIt;
-}
-
-template<class LoadIdBinder>
-void CopyMemberBinding(FRangeMemberBinding Binding, /* in-out */ const FSchemaId*& InnerIdIt, LoadIdBinder& LoadIds, FMemberLoadBinder& Out)
-{
-	FMemberBindType InnermostType = Binding.InnerTypes[Binding.NumRanges - 1];
-	Out.AddRange(MakeArrayView(Binding.RangeBindings, Binding.NumRanges), InnermostType, static_cast<uint32>(Binding.Offset));
-	if (InnermostType.IsStruct())
-	{
-		FStructSchemaId ReadId = static_cast<FStructSchemaId>(*InnerIdIt);
-		FStructSchemaId LoadId = LoadIds.BindLoadId(ReadId, Binding.InnermostId.Get().AsStructBindId());
-		Out.AddInnerId(LoadId);
-		++InnerIdIt;
-	}
-	else
-	{
-		InnerIdIt += (InnermostType.AsLeaf().Bind.Type == ELeafBindType::Enum); // Skip enum schema
-	}
-}
-
-template<class LoadIdBinder>
-void CopyMemberBinding(/* in-out */ FMemberVisitor& BindIt, /* in-out */ const FSchemaId*& InnerIdIt, LoadIdBinder& LoadIds, FMemberLoadBinder& Out)
-{
-	switch (BindIt.PeekKind())
-	{
-		case EMemberKind::Leaf:		CopyMemberBinding(BindIt.GrabLeaf(), InnerIdIt, Out);				break;
-		case EMemberKind::Range:	CopyMemberBinding(BindIt.GrabRange(), InnerIdIt, LoadIds, Out);		break;
-		case EMemberKind::Struct:	CopyMemberBinding(BindIt.GrabStruct(), InnerIdIt, LoadIds, Out);	break;
-		default:					check(false);															break;
-	}
-}
-
-
-template<class LoadIdBinder>
-static void CreateSubsetBindingWithoutEnumIds(const FStructSchema& From, const FSchemaBinding& To, TConstArrayView<FMemberId> ToNames, uint16 NumEnums, LoadIdBinder& LoadIds, SubsetByteArray& Out)
-{
-	check(To.NumMembers == ToNames.Num() + To.HasSuper());
-	check(To.NumMembers >= From.NumMembers);
-	check(IsAligned(Out.Num(), alignof(FSchemaBinding)));
-
-	int32 OutPos = Out.Num();
-	
-	// Allocate and init header 
-	bool bSkipSuperSchema = SkipDeclaredSuperSchema(From.Inheritance);
-	FSchemaBinding Header = { To.DeclId, From.NumMembers, From.NumInnerSchemas - NumEnums - bSkipSuperSchema, From.NumRangeTypes };
-	Out.AddUninitialized(Header.CalculateSize());
-	FSchemaBinding* Schema = new (&Out[OutPos]) FSchemaBinding {Header};
-	
-	// Copy subset of member bindings...
-	FMemberVisitor ToIt(To);
-	FMemberLoadBinder Footer(*Schema);
-	const FSchemaId* InnerIdIt = From.GetInnerSchemas() + bSkipSuperSchema;
-
-	// ...first the unnamed super member...
-	if (From.Inheritance != ESuper::No)
-	{
-		FBindId DeclBindId = ToIt.GrabSuper();
-		if (UsesSuper(From.Inheritance))
-		{
-			FStructType FromType = From.Footer[0].AsStruct(); // To.Members[0].AsStruct().IsDynamic isn't set
-			check(FromType.IsSuper);
-			Footer.AddMember(FMemberBindType(FromType), 0);
-		
-			if (!FromType.IsDynamic)
-			{
-				FStructSchemaId ReadId = static_cast<FStructSchemaId>(*InnerIdIt);
-				FStructSchemaId LoadId;
-				if (From.Inheritance == ESuper::Reused)
-				{
-					LoadId = LoadIds.BindLoadId(ReadId, DeclBindId);
-				}
-				else
-				{
-					check(From.Inheritance == ESuper::Used);
-					// TODO: Handle type erasure in the used super scenario too,
-					// i.e. when reading a base of the declared super, needs UsedBindId.
-					LoadId = ReadId;
-				}
-				Footer.AddInnerId(LoadId);
-				++InnerIdIt;
-			}
-		}
-	}
-
-	// ...then a subset of named members
-	const FMemberId* ToNameIt = ToNames.GetData();
-	for (FMemberId FromName : From.GetMemberNames())
-	{
-		while (FromName != *ToNameIt++)
-		{
-			ToIt.SkipMember();
-			check(ToNameIt != ToNames.end());
-		}
-		
-		CopyMemberBinding(/* in-out */ ToIt, /* in-out */ InnerIdIt, /* in-out */ LoadIds,  /* out */ Footer);	
-	}
-	check(InnerIdIt == From.GetInnerSchemas() + From.NumInnerSchemas);
-}
-
-// @pre No enum members
-template<class LoadIdBinder>
-static void CloneBindingWithReplacedStructIds(const FSchemaId* FromIds, const FSchemaBinding& To, LoadIdBinder& LoadIds, SubsetByteArray& Out)
-{
-	check(IsAligned(Out.Num(), alignof(FSchemaBinding)));
-
-	// Clone runtime FSchemaBinding including inner schema bind ids
-	uint32 Size = To.CalculateSize();
-	Out.AddUninitialized(Size);
-	FSchemaBinding* Schema = reinterpret_cast<FSchemaBinding*>(&Out[Out.Num() - Size]);
-	FMemory::Memcpy(Schema, &To, Size);
-
-	// Replace inner bind ids with batch load ids
-	const FStructSchemaId* ReadIdIt = static_cast<const FStructSchemaId*>(FromIds);
-	FInnerId* InnerIt = const_cast<FInnerId*>(Schema->GetInnerIds());
-	for (FInnerId& Inner : MakeArrayView(InnerIt, To.NumInnerIds))
-	{
-		FBindId MemcopiedBindId = Inner.AsStructBindId();
-		reinterpret_cast<FStructSchemaId&>(Inner) = LoadIds.BindLoadId(*ReadIdIt++, MemcopiedBindId);
-	}
-}
-
-template<class LoadIdBinder>
-[[nodiscard]] static FLoadStructPlan MakeSchemaLoadPlan(const FStructSchema& From, const FSchemaBinding& To, TConstArrayView<FMemberId> ToMemberIds, TConstArrayView<FStructId> ToStructIds, LoadIdBinder& LoadIds, SubsetByteArray& OutSubsetSchemas)
-{
-	uint16 NumEnums = CountEnums(From);
-	if (From.NumMembers < To.NumMembers || NumEnums || HasIncompatibleSupers(From, To, ToStructIds))
-	{
-		CreateSubsetBindingWithoutEnumIds(From, To, ToMemberIds, NumEnums, LoadIds, OutSubsetSchemas);
-	}
-	else
-	{
-		check(From.NumMembers == To.NumMembers);
-		check(From.NumInnerSchemas == To.NumInnerIds);
-		check(From.NumRangeTypes == To.NumInnerRanges);
-
-		if (From.NumInnerSchemas > 0)
-		{
-			CloneBindingWithReplacedStructIds(From.GetInnerSchemas(), To, LoadIds, OutSubsetSchemas);
-		}
-		// else reuse existing bindings
-	}
-
-	// Pointer to created subset load schema will be remapped later
-	return FLoadStructPlan(To, ELeafWidth::B32, !From.IsDense);
-}
-
-[[nodiscard]] static TOptional<FMemcpyLoadPlan> TryMakeMemcpyPlan(const FStructSchema& From, const FSchemaBinding& To, TConstArrayView<FMemberId> ToNames)
-{
-	// Can't memcopy sparse members nor range-bound members nor super structs
-	if (!From.IsDense || From.NumRangeTypes > 0 || UsesSuper(From.Inheritance))
-	{
-		return NullOpt;
-	}
-
-	// Can't memcpy non-contiguous members
-	TConstArrayView<FMemberId> FromNames = From.GetMemberNames();
-	int32 SkipToIdx = ToNames.Find(FromNames[0]);
-	check(SkipToIdx != INDEX_NONE);
-	if (!Algo::Compare(FromNames, ToNames.Slice(SkipToIdx, FromNames.Num())))
-	{	
-		return NullOpt;
-	}
-	
-	// Check all schema members are contiguous leaves
-	const uint32* OffsetIt = To.GetOffsets() + SkipToIdx;
-	const uint32 StartPos = *OffsetIt;
-	uint32 EndPos = *OffsetIt;
-	for (FMemberType Member : From.GetMemberTypes())
-	{
-		// Note: By adding FStructType::IsDense flag and some struct size lookup, memcpying of nested non-dynamic non-custom-bound structs might be possible
-		uint32 Offset = *OffsetIt++;
-		if (Offset != EndPos ||	// Non-contiguous or padded
-			Member.IsStruct() || // Nested dtructs have a skippable size prefix that can't be memcopied
-			Member.AsLeaf().Type == ELeafType::Bool) // bool values are stored as bits, see FBitCacheReader
-		{
-			return NullOpt;
-		}
-
-		EndPos += SizeOf(Member.AsLeaf().Width);
-	}
-
-	return FMemcpyLoadPlan(EndPos - StartPos, StartPos);
-}
-
-template<class LoadIdBinder>
-[[nodiscard]] static FLoadStructPlan MakeLoadPlan(const FStructSchema& From, const FSchemaBinding& To, TConstArrayView<FMemberId> ToMemberIds, TConstArrayView<FStructId> ToStructIds, LoadIdBinder& LoadIds, SubsetByteArray& OutSubsetSchemas)
-{
-	TOptional<FMemcpyLoadPlan> Memcpy = TryMakeMemcpyPlan(From, To, ToMemberIds);
-	return Memcpy ? FLoadStructPlan(Memcpy.GetValue()) : MakeSchemaLoadPlan(From, To, ToMemberIds, ToStructIds, LoadIds, OutSubsetSchemas);
-}
-
 struct FLoadBindingBase : ICustomBinding
 {
 	virtual void SaveCustom(FMemberBuilder& Dst, const void* Src, FBaseline Base, const FSaveContext& Ctx) override final { unimplemented(); }
@@ -849,6 +641,212 @@ private:
 
 		return FLoadStructPlan(GLoadNoop);
 	}
+
+	static void CopyMemberBinding(FLeafMemberBinding Binding, /* in-out */ const FSchemaId*& InnerIdIt, FMemberLoadBinder& Out)
+	{
+		InnerIdIt += (Binding.Leaf.Type == ELeafBindType::Enum); // Skip enum schema
+		Out.AddMember(Binding.Leaf.Pack(), static_cast<uint32>(Binding.Offset));
+	}
+
+	template<class LoadIdBinder>
+	void CopyMemberBinding(FStructMemberBinding Binding, /* in-out */ const FSchemaId*& InnerIdIt, LoadIdBinder& LoadIds, FMemberLoadBinder& Out)
+	{
+		check(!Binding.Type.IsSuper);
+		check(!Binding.Type.IsDynamic); // TODO: Build out dynamic struct support
+		FStructSchemaId LoadId = LoadIds.BindLoadId(static_cast<FStructSchemaId>(*InnerIdIt), Binding.Id);
+		Out.AddMember(FMemberBindType(Binding.Type), static_cast<uint32>(Binding.Offset));
+		Out.AddInnerId(LoadId);
+		++InnerIdIt;
+	}
+
+	template<class LoadIdBinder>
+	void CopyMemberBinding(FRangeMemberBinding Binding, /* in-out */ const FSchemaId*& InnerIdIt, LoadIdBinder& LoadIds, FMemberLoadBinder& Out)
+	{
+		FMemberBindType InnermostType = Binding.InnerTypes[Binding.NumRanges - 1];
+		Out.AddRange(MakeArrayView(Binding.RangeBindings, Binding.NumRanges), InnermostType, static_cast<uint32>(Binding.Offset));
+		if (InnermostType.IsStruct())
+		{
+			FStructSchemaId ReadId = static_cast<FStructSchemaId>(*InnerIdIt);
+			FStructSchemaId LoadId = LoadIds.BindLoadId(ReadId, Binding.InnermostId.Get().AsStructBindId());
+			Out.AddInnerId(LoadId);
+			++InnerIdIt;
+		}
+		else
+		{
+			InnerIdIt += (InnermostType.AsLeaf().Bind.Type == ELeafBindType::Enum); // Skip enum schema
+		}
+	}
+
+	template<class LoadIdBinder>
+	void CopyMemberBinding(/* in-out */ FMemberVisitor& BindIt, /* in-out */ const FSchemaId*& InnerIdIt, LoadIdBinder& LoadIds, FMemberLoadBinder& Out)
+	{
+		switch (BindIt.PeekKind())
+		{
+			case EMemberKind::Leaf:		CopyMemberBinding(BindIt.GrabLeaf(), InnerIdIt, Out);				break;
+			case EMemberKind::Range:	CopyMemberBinding(BindIt.GrabRange(), InnerIdIt, LoadIds, Out);		break;
+			case EMemberKind::Struct:	CopyMemberBinding(BindIt.GrabStruct(), InnerIdIt, LoadIds, Out);	break;
+			default:					check(false);															break;
+		}
+	}
+
+
+	template<class LoadIdBinder>
+	void CreateSubsetBindingWithoutEnumIds(const FStructSchema& From, const FSchemaBinding& To, TConstArrayView<FMemberId> ToNames, uint16 NumEnums, LoadIdBinder& LoadIds, SubsetByteArray& Out)
+	{
+		check(To.NumMembers == ToNames.Num() + To.HasSuper());
+		check(To.NumMembers >= From.NumMembers);
+		check(IsAligned(Out.Num(), alignof(FSchemaBinding)));
+
+		int32 OutPos = Out.Num();
+		
+		// Allocate and init header 
+		bool bSkipSuperSchema = SkipDeclaredSuperSchema(From.Inheritance);
+		FSchemaBinding Header = { To.DeclId, From.NumMembers, From.NumInnerSchemas - NumEnums - bSkipSuperSchema, From.NumRangeTypes };
+		Out.AddUninitialized(Header.CalculateSize());
+		FSchemaBinding* Schema = new (&Out[OutPos]) FSchemaBinding {Header};
+		
+		// Copy subset of member bindings...
+		FMemberVisitor ToIt(To);
+		FMemberLoadBinder Footer(*Schema);
+		const FSchemaId* InnerIdIt = From.GetInnerSchemas() + bSkipSuperSchema;
+
+		// ...first the unnamed super member...
+		if (From.Inheritance != ESuper::No)
+		{
+			FBindId DeclBindId = ToIt.GrabSuper();
+			if (UsesSuper(From.Inheritance))
+			{
+				FStructType FromType = From.Footer[0].AsStruct(); // To.Members[0].AsStruct().IsDynamic isn't set
+				check(FromType.IsSuper);
+				Footer.AddMember(FMemberBindType(FromType), 0);
+			
+				if (!FromType.IsDynamic)
+				{
+					FStructSchemaId ReadId = static_cast<FStructSchemaId>(*InnerIdIt);
+					FStructSchemaId LoadId;
+					if (From.Inheritance == ESuper::Reused)
+					{
+						LoadId = LoadIds.BindLoadId(ReadId, DeclBindId);
+					}
+					else
+					{
+						check(From.Inheritance == ESuper::Used);
+						// TODO: Handle type erasure in the used super scenario too,
+						// i.e. when reading a base of the declared super, needs UsedBindId.
+						LoadId = ReadId;
+					}
+					Footer.AddInnerId(LoadId);
+					++InnerIdIt;
+				}
+			}
+		}
+
+		// ...then a subset of named members
+		const FMemberId* ToNameIt = ToNames.GetData();
+		for (FMemberId FromName : From.GetMemberNames())
+		{
+			while (FromName != *ToNameIt++)
+			{
+				ToIt.SkipMember();
+				check(ToNameIt != ToNames.end());
+			}
+			
+			CopyMemberBinding(/* in-out */ ToIt, /* in-out */ InnerIdIt, /* in-out */ LoadIds,  /* out */ Footer);	
+		}
+		check(InnerIdIt == From.GetInnerSchemas() + From.NumInnerSchemas);
+	}
+
+	// @pre No enum members
+	template<class LoadIdBinder>
+	void CloneBindingWithReplacedStructIds(const FSchemaId* FromIds, const FSchemaBinding& To, LoadIdBinder& LoadIds, SubsetByteArray& Out)
+	{
+		check(IsAligned(Out.Num(), alignof(FSchemaBinding)));
+
+		// Clone runtime FSchemaBinding including inner schema bind ids
+		uint32 Size = To.CalculateSize();
+		Out.AddUninitialized(Size);
+		FSchemaBinding* Schema = reinterpret_cast<FSchemaBinding*>(&Out[Out.Num() - Size]);
+		FMemory::Memcpy(Schema, &To, Size);
+
+		// Replace inner bind ids with batch load ids
+		const FStructSchemaId* ReadIdIt = static_cast<const FStructSchemaId*>(FromIds);
+		FInnerId* InnerIt = const_cast<FInnerId*>(Schema->GetInnerIds());
+		for (FInnerId& Inner : MakeArrayView(InnerIt, To.NumInnerIds))
+		{
+			FBindId MemcopiedBindId = Inner.AsStructBindId();
+			reinterpret_cast<FStructSchemaId&>(Inner) = LoadIds.BindLoadId(*ReadIdIt++, MemcopiedBindId);
+		}
+	}
+
+	template<class LoadIdBinder>
+	[[nodiscard]] FLoadStructPlan MakeSchemaLoadPlan(const FStructSchema& From, const FSchemaBinding& To, TConstArrayView<FMemberId> ToMemberIds, TConstArrayView<FStructId> ToStructIds, LoadIdBinder& LoadIds, SubsetByteArray& OutSubsetSchemas)
+	{
+		uint16 NumEnums = CountEnums(From);
+		if (From.NumMembers < To.NumMembers || NumEnums || HasIncompatibleSupers(From, To, ToStructIds))
+		{
+			CreateSubsetBindingWithoutEnumIds(From, To, ToMemberIds, NumEnums, LoadIds, OutSubsetSchemas);
+		}
+		else
+		{
+			check(From.NumMembers == To.NumMembers);
+			check(From.NumInnerSchemas == To.NumInnerIds);
+			check(From.NumRangeTypes == To.NumInnerRanges);
+
+			if (From.NumInnerSchemas > 0)
+			{
+				CloneBindingWithReplacedStructIds(From.GetInnerSchemas(), To, LoadIds, OutSubsetSchemas);
+			}
+			// else reuse existing bindings
+		}
+
+		// Pointer to created subset load schema will be remapped later
+		return FLoadStructPlan(To, ELeafWidth::B32, !From.IsDense);
+	}
+
+	[[nodiscard]] TOptional<FMemcpyLoadPlan> TryMakeMemcpyPlan(const FStructSchema& From, const FSchemaBinding& To, TConstArrayView<FMemberId> ToNames)
+	{
+		// Can't memcopy sparse members nor range-bound members nor super structs
+		if (!From.IsDense || From.NumRangeTypes > 0 || UsesSuper(From.Inheritance))
+		{
+			return NullOpt;
+		}
+
+		// Can't memcpy non-contiguous members
+		TConstArrayView<FMemberId> FromNames = From.GetMemberNames();
+		int32 SkipToIdx = ToNames.Find(FromNames[0]);
+		check(SkipToIdx != INDEX_NONE);
+		if (!Algo::Compare(FromNames, ToNames.Slice(SkipToIdx, FromNames.Num())))
+		{	
+			return NullOpt;
+		}
+		
+		// Check all schema members are contiguous leaves
+		const uint32* OffsetIt = To.GetOffsets() + SkipToIdx;
+		const uint32 StartPos = *OffsetIt;
+		uint32 EndPos = *OffsetIt;
+		for (FMemberType Member : From.GetMemberTypes())
+		{
+			// Note: By adding FStructType::IsDense flag and some struct size lookup, memcpying of nested non-dynamic non-custom-bound structs might be possible
+			uint32 Offset = *OffsetIt++;
+			if (Offset != EndPos ||	// Non-contiguous or padded
+				Member.IsStruct() || // Nested dtructs have a skippable size prefix that can't be memcopied
+				Member.AsLeaf().Type == ELeafType::Bool) // bool values are stored as bits, see FBitCacheReader
+			{
+				return NullOpt;
+			}
+
+			EndPos += SizeOf(Member.AsLeaf().Width);
+		}
+
+		return FMemcpyLoadPlan(EndPos - StartPos, StartPos);
+	}
+
+	template<class LoadIdBinder>
+	[[nodiscard]] FLoadStructPlan MakeLoadPlan(const FStructSchema& From, const FSchemaBinding& To, TConstArrayView<FMemberId> ToMemberIds, TConstArrayView<FStructId> ToStructIds, LoadIdBinder& LoadIds, SubsetByteArray& OutSubsetSchemas)
+	{
+		TOptional<FMemcpyLoadPlan> Memcpy = TryMakeMemcpyPlan(From, To, ToMemberIds);
+		return Memcpy ? FLoadStructPlan(Memcpy.GetValue()) : MakeSchemaLoadPlan(From, To, ToMemberIds, ToStructIds, LoadIds, OutSubsetSchemas);
+	}
 };
 
 FLoadBatchPtr CreateLoadPlans(FSchemaBatchId BatchId, const FCustomBindings& Customs, const FSchemaBindings& Schemas, TConstArrayView<FStructId> RuntimeIds, ESchemaFormat Format)
