PlainProps: Move free functions into FLoadPlanner member functions - part 2.
Add LoadIds as a member.
Stop passing around LoadIds, RuntimeIds and SubsetSchemaData as parameters.

#rnx

diff --git a/Engine/Plugins/Experimental/PlainProps/Source/Private/PlainPropsLoad.cpp b/Engine/Plugins/Experimental/PlainProps/Source/Private/PlainPropsLoad.cpp
--- a/Engine/Plugins/Experimental/PlainProps/Source/Private/PlainPropsLoad.cpp
+++ b/Engine/Plugins/Experimental/PlainProps/Source/Private/PlainPropsLoad.cpp
@@ -263,8 +263,6 @@ void FLoadBatchDeleter::operator()(FLoadBatch* Batch) const
 	FMemory::Free(Batch);
 }
 
-using SubsetByteArray = TArray<uint8, TInlineAllocator<1024>>;
-
 ////////////////////////////////////////////////////////////////////////////
 
 struct FMemberLoadBinder : FMemberBinderBase
@@ -334,21 +332,25 @@ static TOptional<TConstArrayView<T>> GetContiguousSubset(TConstArrayView<T> View
 	return NullOpt;
 }
 
+template<typename LoadIdBinder>
 struct FLoadPlanner
 {
+	static constexpr bool bIsLoadIdBinderDummy = std::is_same_v<LoadIdBinder, FLoadIdBinderDummy>;
+
 	// Constructor parameters
 	const FSchemaBatchId							BatchId;
 	const FCustomBindings&							Customs;
 	const FSchemaBindings&							Schemas;
 	const TConstArrayView<FStructId>				RuntimeIds;
+	LoadIdBinder									LoadIds;
 
 	// Temporary data structures	
 	TArray<FLoadStructPlan, TInlineAllocator<256>>	Plans;
 	TArray<uint32, TInlineAllocator<256>>			SubsetSchemaSizes;
-	SubsetByteArray									SubsetSchemaData;
+	TArray<uint8, TInlineAllocator<1024>>			SubsetSchemaData;
 	TSet<FStructSchemaId>							UnboundSaveIds;
 	
-	FLoadBatchPtr CreatePlans(ESchemaFormat Format)
+	FLoadBatchPtr CreatePlans()
 	{
 		check(NumStructSchemas(BatchId) == RuntimeIds.Num());
 
@@ -357,64 +359,61 @@ struct FLoadPlanner
 		Plans.SetNumUninitialized(NumPlans);
 		SubsetSchemaSizes.SetNumUninitialized(NumPlans);
 
-		if (Format == ESchemaFormat::InMemoryNames)
+		if constexpr (bIsLoadIdBinderDummy)
 		{
-			FLoadIdBinderDummy LoadIds{Schemas.GetDebug()};
 			for (uint32 Idx = 0; Idx < NumPlans; ++Idx)
 			{
 				FLoadIdMapping Mapping;
 				Mapping.ReadId = { Idx };
 				Mapping.LoadId = { Idx };
 				Mapping.Id = FBindId(RuntimeIds[Idx]);
-				CreatePlan(Mapping, LoadIds);
+				CreatePlan(Mapping);
 			}
-
-			return CreateBatch(LoadIds);
+			return CreateBatch();
 		}
-	
-		TConstArrayView<FDeclId> DeclIds(static_cast<const FDeclId*>(RuntimeIds.GetData()), RuntimeIds.Num());
-		FLoadIdBinder LoadIds(DeclIds, Schemas.GetDebug());
-		for (uint32 Idx = 0; Idx < NumPlans; ++Idx)
+		else
 		{
-			FLoadIdMapping Mapping;
-			Mapping.ReadId = { Idx };
-			Mapping.LoadId = { Idx };
-			Mapping.Id = UpCast(DeclIds[Idx]);
-			CreatePlan(Mapping, LoadIds);
-		}
+			for (uint32 Idx = 0; Idx < NumPlans; ++Idx)
+			{
+				FLoadIdMapping Mapping;
+				Mapping.ReadId = { Idx };
+				Mapping.LoadId = { Idx };
+				Mapping.Id = FBindId(RuntimeIds[Idx]);
+				CreatePlan(Mapping);
+			}
 
-		// Make load plans for type-erased/ExplicitBindName structs needed by already created load plans
-		if (LoadIds.NumMappings() > 0)
-		{
-			Plans.Reserve(NumPlans + LoadIds.NumMappings());
-			SubsetSchemaSizes.Reserve(NumPlans + LoadIds.NumMappings());
-			for (int32 Idx = 0; Idx < LoadIds.NumMappings(); ++Idx)
+			// Make load plans for type-erased/ExplicitBindName structs needed by already created load plans
+			if (LoadIds.NumMappings() > 0)
 			{
-				check(LoadIds.GetMapping(Idx).LoadId.Idx == static_cast<uint32>(Plans.Num()));
-				Plans.AddUninitialized();
-				SubsetSchemaSizes.AddUninitialized();
-				CreatePlan(LoadIds.GetMapping(Idx), /* might grow */ LoadIds);
+				Plans.Reserve(NumPlans + LoadIds.NumMappings());
+				SubsetSchemaSizes.Reserve(NumPlans + LoadIds.NumMappings());
+				for (int32 Idx = 0; Idx < LoadIds.NumMappings(); ++Idx)
+				{
+					check(LoadIds.GetMapping(Idx).LoadId.Idx == static_cast<uint32>(Plans.Num()));
+					Plans.AddUninitialized();
+					SubsetSchemaSizes.AddUninitialized();
+					CreatePlan(LoadIds.GetMapping(Idx)); // LoadIds might grow
+				}
+
+				// Verify that all UnboundSaveIds were bound by some load plan
+				for (int32 Idx = 0; !UnboundSaveIds.IsEmpty() && Idx < LoadIds.NumMappings(); ++Idx)
+				{
+					UnboundSaveIds.Remove(LoadIds.GetMapping(Idx).ReadId);
+				}
 			}
 
-			// Verify that all UnboundSaveIds were bound by some load plan
-			for (int32 Idx = 0; !UnboundSaveIds.IsEmpty() && Idx < LoadIds.NumMappings(); ++Idx)
+			for (FStructSchemaId UnboundSaveId : UnboundSaveIds)
 			{
-				UnboundSaveIds.Remove(LoadIds.GetMapping(Idx).ReadId);
+				checkf(false, TEXT("Unbound struct '%s' can't be loaded"), *Schemas.GetDebug().Print(RuntimeIds[UnboundSaveId.Idx]));
 			}
-		}
 
-		for (FStructSchemaId UnboundSaveId : UnboundSaveIds)
-		{
-			checkf(false, TEXT("Unbound struct '%s' can't be loaded"), *Schemas.GetDebug().Print(RuntimeIds[UnboundSaveId.Idx]));
+			// Allocate load batch, copy plans and subset schemas, and fixup subset schema plans
+			return CreateBatch();
 		}
-
-		// Allocate load batch, copy plans and subset schemas, and fixup subset schema plans
-		return CreateBatch(LoadIds);
 	}
 private:
 
-	template<typename LoadIdBinder>
-	FLoadBatchPtr CreateBatch(const LoadIdBinder& LoadIds)
+	FLoadBatchPtr CreateBatch() const
 	{
 		const uint32 NumPlans = Plans.Num();
 		const uint32 NumMappings = LoadIds.NumMappings();
@@ -471,11 +470,10 @@ private:
 		return FLoadBatchPtr(Out);
 	}
 
-	template<class LoadIdBinder>
-	void CreatePlan(FLoadIdMapping Mapping, LoadIdBinder& LoadIds)
+	void CreatePlan(FLoadIdMapping Mapping)
 	{
 		int32 SubsetSchemaOffset = SubsetSchemaData.Num();	
-		Plans[Mapping.LoadId.Idx] = CreatePlanInner(Mapping, LoadIds);
+		Plans[Mapping.LoadId.Idx] = CreatePlanInner(Mapping);
 		SubsetSchemaSizes[Mapping.LoadId.Idx] = static_cast<uint32>(SubsetSchemaData.Num() - SubsetSchemaOffset);
 	}
 
@@ -520,7 +518,7 @@ private:
 		{}
 	};
 
-	[[nodiscard]] FCustomLoadPlan* CreateCustomLoadIdsPlan(const ICustomBinding* Custom, const FStructSchema& Schema, TConstArrayView<FInnerStruct> Inners, /* in-out */ FLoadIdBinder& LoadIds)
+	[[nodiscard]] FCustomLoadPlan* CreateCustomLoadIdsPlan(const ICustomBinding* Custom, const FStructSchema& Schema, TConstArrayView<FInnerStruct> Inners)
 	{
 		const uint32 Num = static_cast<uint32>(Inners.Num());
 		check(Num);
@@ -589,13 +587,12 @@ private:
 		return TmpSubset;
 	}
 
-	template<class LoadIdBinder>
-	FLoadStructPlan CreatePlanInner(FLoadIdMapping Mapping, LoadIdBinder& LoadIds)
+	FLoadStructPlan CreatePlanInner(FLoadIdMapping Mapping)
 	{
 		const FStructSchema& From = ResolveStructSchema(BatchId, Mapping.ReadId);
 		
 		TConstArrayView<FInnerStruct> LoweredInners;
-		if constexpr (std::is_same_v<LoadIdBinder, FLoadIdBinderDummy>)
+		if constexpr (bIsLoadIdBinderDummy)
 		{
 			if (const ICustomBinding* Custom = Customs.FindStruct(Mapping.Id))
 			{
@@ -621,7 +618,7 @@ private:
 				return FLoadStructPlan(*Custom);
 			}
 			
-			return FLoadStructPlan(*CreateCustomLoadIdsPlan(Custom, From, LoweredMembers, LoadIds));
+			return FLoadStructPlan(*CreateCustomLoadIdsPlan(Custom, From, LoweredMembers));
 		}
 
 		if (From.NumMembers)
@@ -631,7 +628,7 @@ private:
 			{
 				// Possible optimization - some simple memcpy cases doesn't need to resolve the declaration
 				TConstArrayView<FMemberId> ToMemberIds = Decl->GetMemberOrder();
-				return MakeLoadPlan(From, *To, ToMemberIds, RuntimeIds, LoadIds, /* out */ SubsetSchemaData);	
+				return MakeLoadPlan(From, *To, ToMemberIds);	
 			}
 			
 			// Type-erased structs
@@ -648,8 +645,7 @@ private:
 		Out.AddMember(Binding.Leaf.Pack(), static_cast<uint32>(Binding.Offset));
 	}
 
-	template<class LoadIdBinder>
-	void CopyMemberBinding(FStructMemberBinding Binding, /* in-out */ const FSchemaId*& InnerIdIt, LoadIdBinder& LoadIds, FMemberLoadBinder& Out)
+	void CopyMemberBinding(FStructMemberBinding Binding, /* in-out */ const FSchemaId*& InnerIdIt, FMemberLoadBinder& Out)
 	{
 		check(!Binding.Type.IsSuper);
 		check(!Binding.Type.IsDynamic); // TODO: Build out dynamic struct support
@@ -659,8 +655,7 @@ private:
 		++InnerIdIt;
 	}
 
-	template<class LoadIdBinder>
-	void CopyMemberBinding(FRangeMemberBinding Binding, /* in-out */ const FSchemaId*& InnerIdIt, LoadIdBinder& LoadIds, FMemberLoadBinder& Out)
+	void CopyMemberBinding(FRangeMemberBinding Binding, /* in-out */ const FSchemaId*& InnerIdIt, FMemberLoadBinder& Out)
 	{
 		FMemberBindType InnermostType = Binding.InnerTypes[Binding.NumRanges - 1];
 		Out.AddRange(MakeArrayView(Binding.RangeBindings, Binding.NumRanges), InnermostType, static_cast<uint32>(Binding.Offset));
@@ -677,33 +672,31 @@ private:
 		}
 	}
 
-	template<class LoadIdBinder>
-	void CopyMemberBinding(/* in-out */ FMemberVisitor& BindIt, /* in-out */ const FSchemaId*& InnerIdIt, LoadIdBinder& LoadIds, FMemberLoadBinder& Out)
+	void CopyMemberBinding(/* in-out */ FMemberVisitor& BindIt, /* in-out */ const FSchemaId*& InnerIdIt, FMemberLoadBinder& Out)
 	{
 		switch (BindIt.PeekKind())
 		{
-			case EMemberKind::Leaf:		CopyMemberBinding(BindIt.GrabLeaf(), InnerIdIt, Out);				break;
-			case EMemberKind::Range:	CopyMemberBinding(BindIt.GrabRange(), InnerIdIt, LoadIds, Out);		break;
-			case EMemberKind::Struct:	CopyMemberBinding(BindIt.GrabStruct(), InnerIdIt, LoadIds, Out);	break;
-			default:					check(false);															break;
+			case EMemberKind::Leaf:		CopyMemberBinding(BindIt.GrabLeaf(), InnerIdIt, Out);	break;
+			case EMemberKind::Range:	CopyMemberBinding(BindIt.GrabRange(), InnerIdIt, Out);	break;
+			case EMemberKind::Struct:	CopyMemberBinding(BindIt.GrabStruct(), InnerIdIt, Out);	break;
+			default:					check(false);											break;
 		}
 	}
 
 
-	template<class LoadIdBinder>
-	void CreateSubsetBindingWithoutEnumIds(const FStructSchema& From, const FSchemaBinding& To, TConstArrayView<FMemberId> ToNames, uint16 NumEnums, LoadIdBinder& LoadIds, SubsetByteArray& Out)
+	void CreateSubsetBindingWithoutEnumIds(const FStructSchema& From, const FSchemaBinding& To, TConstArrayView<FMemberId> ToNames, uint16 NumEnums)
 	{
 		check(To.NumMembers == ToNames.Num() + To.HasSuper());
 		check(To.NumMembers >= From.NumMembers);
-		check(IsAligned(Out.Num(), alignof(FSchemaBinding)));
+		check(IsAligned(SubsetSchemaData.Num(), alignof(FSchemaBinding)));
 
-		int32 OutPos = Out.Num();
+		int32 OutPos = SubsetSchemaData.Num();
 		
 		// Allocate and init header 
 		bool bSkipSuperSchema = SkipDeclaredSuperSchema(From.Inheritance);
 		FSchemaBinding Header = { To.DeclId, From.NumMembers, From.NumInnerSchemas - NumEnums - bSkipSuperSchema, From.NumRangeTypes };
-		Out.AddUninitialized(Header.CalculateSize());
-		FSchemaBinding* Schema = new (&Out[OutPos]) FSchemaBinding {Header};
+		SubsetSchemaData.AddUninitialized(Header.CalculateSize());
+		FSchemaBinding* Schema = new (&SubsetSchemaData[OutPos]) FSchemaBinding {Header};
 		
 		// Copy subset of member bindings...
 		FMemberVisitor ToIt(To);
@@ -751,21 +744,20 @@ private:
 				check(ToNameIt != ToNames.end());
 			}
 			
-			CopyMemberBinding(/* in-out */ ToIt, /* in-out */ InnerIdIt, /* in-out */ LoadIds,  /* out */ Footer);	
+			CopyMemberBinding(/* in-out */ ToIt, /* in-out */ InnerIdIt, /* out */ Footer);	
 		}
 		check(InnerIdIt == From.GetInnerSchemas() + From.NumInnerSchemas);
 	}
 
 	// @pre No enum members
-	template<class LoadIdBinder>
-	void CloneBindingWithReplacedStructIds(const FSchemaId* FromIds, const FSchemaBinding& To, LoadIdBinder& LoadIds, SubsetByteArray& Out)
+	void CloneBindingWithReplacedStructIds(const FSchemaId* FromIds, const FSchemaBinding& To)
 	{
-		check(IsAligned(Out.Num(), alignof(FSchemaBinding)));
+		check(IsAligned(SubsetSchemaData.Num(), alignof(FSchemaBinding)));
 
 		// Clone runtime FSchemaBinding including inner schema bind ids
 		uint32 Size = To.CalculateSize();
-		Out.AddUninitialized(Size);
-		FSchemaBinding* Schema = reinterpret_cast<FSchemaBinding*>(&Out[Out.Num() - Size]);
+		SubsetSchemaData.AddUninitialized(Size);
+		FSchemaBinding* Schema = reinterpret_cast<FSchemaBinding*>(&SubsetSchemaData[SubsetSchemaData.Num() - Size]);
 		FMemory::Memcpy(Schema, &To, Size);
 
 		// Replace inner bind ids with batch load ids
@@ -778,13 +770,12 @@ private:
 		}
 	}
 
-	template<class LoadIdBinder>
-	[[nodiscard]] FLoadStructPlan MakeSchemaLoadPlan(const FStructSchema& From, const FSchemaBinding& To, TConstArrayView<FMemberId> ToMemberIds, TConstArrayView<FStructId> ToStructIds, LoadIdBinder& LoadIds, SubsetByteArray& OutSubsetSchemas)
+	[[nodiscard]] FLoadStructPlan MakeSchemaLoadPlan(const FStructSchema& From, const FSchemaBinding& To, TConstArrayView<FMemberId> ToMemberIds)
 	{
 		uint16 NumEnums = CountEnums(From);
-		if (From.NumMembers < To.NumMembers || NumEnums || HasIncompatibleSupers(From, To, ToStructIds))
+		if (From.NumMembers < To.NumMembers || NumEnums || HasIncompatibleSupers(From, To, RuntimeIds))
 		{
-			CreateSubsetBindingWithoutEnumIds(From, To, ToMemberIds, NumEnums, LoadIds, OutSubsetSchemas);
+			CreateSubsetBindingWithoutEnumIds(From, To, ToMemberIds, NumEnums);
 		}
 		else
 		{
@@ -794,7 +785,7 @@ private:
 
 			if (From.NumInnerSchemas > 0)
 			{
-				CloneBindingWithReplacedStructIds(From.GetInnerSchemas(), To, LoadIds, OutSubsetSchemas);
+				CloneBindingWithReplacedStructIds(From.GetInnerSchemas(), To);
 			}
 			// else reuse existing bindings
 		}
@@ -841,17 +832,26 @@ private:
 		return FMemcpyLoadPlan(EndPos - StartPos, StartPos);
 	}
 
-	template<class LoadIdBinder>
-	[[nodiscard]] FLoadStructPlan MakeLoadPlan(const FStructSchema& From, const FSchemaBinding& To, TConstArrayView<FMemberId> ToMemberIds, TConstArrayView<FStructId> ToStructIds, LoadIdBinder& LoadIds, SubsetByteArray& OutSubsetSchemas)
+	[[nodiscard]] FLoadStructPlan MakeLoadPlan(const FStructSchema& From, const FSchemaBinding& To, TConstArrayView<FMemberId> ToMemberIds)
 	{
 		TOptional<FMemcpyLoadPlan> Memcpy = TryMakeMemcpyPlan(From, To, ToMemberIds);
-		return Memcpy ? FLoadStructPlan(Memcpy.GetValue()) : MakeSchemaLoadPlan(From, To, ToMemberIds, ToStructIds, LoadIds, OutSubsetSchemas);
+		return Memcpy ? FLoadStructPlan(Memcpy.GetValue()) : MakeSchemaLoadPlan(From, To, ToMemberIds);
 	}
 };
 
 FLoadBatchPtr CreateLoadPlans(FSchemaBatchId BatchId, const FCustomBindings& Customs, const FSchemaBindings& Schemas, TConstArrayView<FStructId> RuntimeIds, ESchemaFormat Format)
 {
-	return FLoadPlanner{BatchId, Customs, Schemas, RuntimeIds}.CreatePlans(Format);
+	if (Format == ESchemaFormat::InMemoryNames)
+	{
+		FLoadIdBinderDummy LoadIds{Schemas.GetDebug()};
+		return FLoadPlanner{BatchId, Customs, Schemas, RuntimeIds, LoadIds}.CreatePlans();
+	}
+	else
+	{
+		TConstArrayView<FDeclId> DeclIds(static_cast<const FDeclId*>(RuntimeIds.GetData()), RuntimeIds.Num());
+		FLoadIdBinder LoadIds(DeclIds, Schemas.GetDebug());
+		return FLoadPlanner{BatchId, Customs, Schemas, RuntimeIds, LoadIds}.CreatePlans();
+	}
 }
 
 ////////////////////////////////////////////////////////////////////////////
