PlainProps add a bunch of helper functions and utils used by upgraders

#rnx
#Synced-CL 47103930

diff --git a/Engine/Plugins/Experimental/PlainProps/Source/Private/PlainPropsBind.cpp b/Engine/Plugins/Experimental/PlainProps/Source/Private/PlainPropsBind.cpp
--- a/Engine/Plugins/Experimental/PlainProps/Source/Private/PlainPropsBind.cpp
+++ b/Engine/Plugins/Experimental/PlainProps/Source/Private/PlainPropsBind.cpp
@@ -674,7 +674,29 @@ FDeclId FBindDeclarations::Lower(FBindId Id) const
 
 //////////////////////////////////////////////////////////////////////////
 
-TArray<FStructId> IndexRuntimeIds(const FSchemaBatch& Schemas, FIdIndexerBase& Indexer)
+TArray<FInnerId> IndexAllRuntimeIds(const FSchemaBatch& Schemas, FIdIndexerBase& Indexer)
+{
+	const uint8* Base = reinterpret_cast<const uint8*>(&Schemas);
+	const uint32* Offsets = Schemas.SchemaOffsets;
+
+	TArray<FInnerId> Out;
+	Out.SetNumUninitialized(Schemas.NumSchemas);
+	for (uint32 Idx = 0; Idx < Schemas.NumStructSchemas; ++Idx)
+	{
+		const FStructSchema& Schema = *reinterpret_cast<const FStructSchema*>(Base + Offsets[Idx]);
+		Out[Idx] = FInnerId(Indexer.IndexDeclId(Schema.Type));
+	}
+
+	for (uint32 Idx = Schemas.NumStructSchemas; Idx < Schemas.NumSchemas; ++Idx)
+	{
+		const FEnumSchema& Schema = *reinterpret_cast<const FEnumSchema*>(Base + Offsets[Idx]);
+		Out[Idx] = FInnerId(Indexer.IndexEnum(Schema.Type));
+	}
+
+	return Out;
+}
+
+TArray<FStructId> IndexStructRuntimeIds(const FSchemaBatch& Schemas, FIdIndexerBase& Indexer)
 {
 	const uint8* Base = reinterpret_cast<const uint8*>(&Schemas);
 	const uint32* Offsets = Schemas.SchemaOffsets;
diff --git a/Engine/Plugins/Experimental/PlainProps/Source/Private/PlainPropsBuild.cpp b/Engine/Plugins/Experimental/PlainProps/Source/Private/PlainPropsBuild.cpp
--- a/Engine/Plugins/Experimental/PlainProps/Source/Private/PlainPropsBuild.cpp
+++ b/Engine/Plugins/Experimental/PlainProps/Source/Private/PlainPropsBuild.cpp
@@ -71,6 +71,13 @@ FBuiltRange* FBuiltRange::Create(FScratchAllocator& Scratch, uint64 NumItems, SI
 	return Out;
 }
 
+FBuiltRange* FBuiltRange::Create(FScratchAllocator& Scratch, uint64 NumItems, SIZE_T ItemSize, const void* Items)
+{
+	FBuiltRange* Out = Create(Scratch, NumItems, ItemSize);
+	FMemory::Memcpy(Out->Data, Items, NumItems * ItemSize);
+	return Out;
+}
+
 //////////////////////////////////////////////////////////////////////////
 
 FMemberSchema MakeNestedRangeSchema(FScratchAllocator& Scratch, ESizeType SizeType, FMemberSchema InnerRangeSchema)
@@ -88,14 +95,7 @@ FMemberSchema MakeNestedRangeSchema(FScratchAllocator& Scratch, ESizeType SizeTy
 
 FBuiltRange* CloneLeaves(FScratchAllocator& Scratch, uint64 Num, const void* InData, SIZE_T LeafSize)
 {
-	if (Num == 0)
-	{
-		return nullptr;
-	}
-
-	FBuiltRange* Out = FBuiltRange::Create(Scratch, Num, LeafSize);
-	FMemory::Memcpy(Out->Data, InData, Num * LeafSize);
-	return Out;
+	return FBuiltRange::CreateIf(Scratch, Num, LeafSize, InData);
 }
 
 //	template<typename FloatType>
@@ -283,8 +283,7 @@ FTypedRange FNestedRangeBuilder::BuildAndReset(FScratchAllocator& Scratch, ESize
 	FBuiltRange* Out = nullptr;
 	if (int64 Num = Ranges.Num())
 	{
-		Out = FBuiltRange::Create(Scratch, Num, sizeof(FBuiltRange*));
-		FMemory::Memcpy(Out->Data, Ranges.GetData(), Ranges.NumBytes());
+		Out = FBuiltRange::Create(Scratch, Num, sizeof(FBuiltRange*), Ranges.GetData());
 		Ranges.Reset();
 	}
 
diff --git a/Engine/Plugins/Experimental/PlainProps/Source/Private/PlainPropsDeclare.cpp b/Engine/Plugins/Experimental/PlainProps/Source/Private/PlainPropsDeclare.cpp
--- a/Engine/Plugins/Experimental/PlainProps/Source/Private/PlainPropsDeclare.cpp
+++ b/Engine/Plugins/Experimental/PlainProps/Source/Private/PlainPropsDeclare.cpp
@@ -80,14 +80,22 @@ static void ValidateDeclaration(const FEnumDeclaration& Enum)
 	TSet<uint64, DefaultKeyFuncs<uint64>,TInlineSetAllocator<64>> Constants;
 	for (FEnumerator E : Enum.GetEnumerators())
 	{
-		//checkf(FMath::FloorLog2_64(E.Constant) < 8 * SizeOf(Enum.Width), TEXT("Enumerator constant larger than declared width"));
-
 		bool bDeclared;
 		Names.FindOrAdd(E.Name.Idx, /* out*/ &bDeclared);
 		checkf(!bDeclared, TEXT("Enumerator name declared twice"));
 		Constants.FindOrAdd(E.Constant, /* out*/ &bDeclared);
 		checkf(!bDeclared, TEXT("Enumerator constant declared twice"));
 	}
+
+	if (Enum.Width)
+	{
+		uint64 OverflowMasks[4] = { ~uint64(0) << 8, ~uint64(0) << 16, ~uint64(0) << 32, 0 };
+		uint64 OverflowMask = OverflowMasks[(uint8)Enum.Width.Get()];
+		for (FEnumerator E : Enum.GetEnumerators())
+		{
+			checkf(!(E.Constant & OverflowMask), TEXT("Enumerator constant is larger than declared width"));
+		}
+	}
 }
 
 template<typename T>
diff --git a/Engine/Plugins/Experimental/PlainProps/Source/Private/PlainPropsInternalBind.h b/Engine/Plugins/Experimental/PlainProps/Source/Private/PlainPropsInternalBind.h
--- a/Engine/Plugins/Experimental/PlainProps/Source/Private/PlainPropsInternalBind.h
+++ b/Engine/Plugins/Experimental/PlainProps/Source/Private/PlainPropsInternalBind.h
@@ -104,4 +104,24 @@ struct FMemberBinderBase
 	}
 }
 
+////////////////////////////////////////////////////////////////////////////////////////////////
+
+// FSchemaId -> FInnerId lookup, see IndexAllRuntimeIds
+struct FInnerIds
+{
+	TConstArrayView<FInnerId>		Ids;
+
+	FInnerId						operator[](FSchemaId Schema) const					{ return Ids[Schema.Idx]; }
+	FEnumId							operator[](FEnumSchemaId Schema) const				{ return Ids[Schema.Idx].AsEnum(); }
+	FDeclId							operator[](FStructSchemaId Schema) const			{ return Ids[Schema.Idx].AsStructDeclId(); }
+	FOptionalInnerId				operator[](FOptionalSchemaId Schema) const			{ return Schema ? ToOptional((*this)[Schema.Get()]) : NoId; }
+	FOptionalEnumId					operator[](FOptionalEnumSchemaId Schema) const		{ return Schema ? ToOptional((*this)[Schema.Get()]) : NoId; }
+	FOptionalDeclId					operator[](FOptionalStructSchemaId Schema) const	{ return Schema ? ToOptional((*this)[Schema.Get()]) : NoId; }
+	
+	int32							Num() const											{ return Ids.Num(); }
+	TConstArrayView<FDeclId>		GetDeclIds(uint32 NumStructs) const					{ return MakeArrayView(reinterpret_cast<const FDeclId*>(Ids.GetData()), NumStructs); }
+	TConstArrayView<FStructId>		GetStructIds(uint32 NumStructs) const				{ return MakeArrayView(reinterpret_cast<const FStructId*>(Ids.GetData()), NumStructs); }
+	TConstArrayView<FEnumId>		GetEnumIds(uint32 NumEnums) const					{ return MakeArrayView(reinterpret_cast<const FEnumId*>(Ids.GetData() + Ids.Num() - NumEnums), NumEnums); }
+};
+
 } // namespace PlainProps
\ No newline at end of file
diff --git a/Engine/Plugins/Experimental/PlainProps/Source/Private/PlainPropsInternalBuild.h b/Engine/Plugins/Experimental/PlainProps/Source/Private/PlainPropsInternalBuild.h
--- a/Engine/Plugins/Experimental/PlainProps/Source/Private/PlainPropsInternalBuild.h
+++ b/Engine/Plugins/Experimental/PlainProps/Source/Private/PlainPropsInternalBuild.h
@@ -21,7 +21,9 @@ struct FBuiltRange
 {
 	~FBuiltRange() = delete; // Allocated in FScratchAllocator
 
-	[[nodiscard]] static FBuiltRange*					Create(FScratchAllocator& Allocator, uint64 NumItems, SIZE_T ItemSize);
+	[[nodiscard]] static FBuiltRange*					Create(FScratchAllocator& Scratch, uint64 NumItems, SIZE_T ItemSize);
+	[[nodiscard]] static FBuiltRange*					Create(FScratchAllocator& Scratch, uint64 NumItems, SIZE_T ItemSize, const void* Items);
+	[[nodiscard]] static FBuiltRange*					CreateIf(FScratchAllocator& Scratch, uint64 NumItems, SIZE_T ItemSize, const void* Items) { return NumItems ? Create(Scratch, NumItems, ItemSize, Items) : nullptr; }
 
 	uint64												Num;
 	uint8												Data[0];
diff --git a/Engine/Plugins/Experimental/PlainProps/Source/Private/PlainPropsInternalRead.h b/Engine/Plugins/Experimental/PlainProps/Source/Private/PlainPropsInternalRead.h
--- a/Engine/Plugins/Experimental/PlainProps/Source/Private/PlainPropsInternalRead.h
+++ b/Engine/Plugins/Experimental/PlainProps/Source/Private/PlainPropsInternalRead.h
@@ -92,6 +92,7 @@ struct TSchemaRange
 	T&								First() const					{ return *reinterpret_cast<T*>(Base + Offsets[0]); }
 	T&								Last() const					{ return *reinterpret_cast<T*>(Base + Offsets.Last()); }
 	T&								operator[](uint32 Idx) const	{ return *reinterpret_cast<T*>(Base + Offsets[Idx]); }
+	int32							Num() const						{ return Offsets.Num(); }
 
 	uintptr_t						Base;
 	TConstArrayView<uint32>			Offsets;
diff --git a/Engine/Plugins/Experimental/PlainProps/Source/Public/PlainPropsBind.h b/Engine/Plugins/Experimental/PlainProps/Source/Public/PlainPropsBind.h
--- a/Engine/Plugins/Experimental/PlainProps/Source/Public/PlainPropsBind.h
+++ b/Engine/Plugins/Experimental/PlainProps/Source/Public/PlainPropsBind.h
@@ -644,7 +644,9 @@ struct FBindDeclarations final : IDeclarations
 //////////////////////////////////////////////////////////////////////////
 
 // Save -> load struct ids for ESchemaFormat::InMemoryNames, alternative to side-channel with ExtractRuntimeIds()
-[[nodiscard]] PLAINPROPS_API TArray<FStructId> IndexRuntimeIds(const FSchemaBatch& Schemas,  FIdIndexerBase& Indexer);
+[[nodiscard]] PLAINPROPS_API TArray<FStructId>			IndexStructRuntimeIds(const FSchemaBatch& Schemas,  FIdIndexerBase& Indexer);
+[[nodiscard]] PLAINPROPS_API TArray<FInnerId>			IndexAllRuntimeIds(const FSchemaBatch& Schemas,  FIdIndexerBase& Indexer);
+[[nodiscard]] PLAINPROPS_API TConstArrayView<FStructId>	GetStructRuntimeIds(TConstArrayView<FInnerId> AllIds, uint32 NumStructSchemas);
 
 // Save -> load ids for ESchemaFormat::StableNames
 struct FIdBinding
diff --git a/Engine/Plugins/Experimental/PlainProps/Source/Public/PlainPropsBuild.h b/Engine/Plugins/Experimental/PlainProps/Source/Public/PlainPropsBuild.h
--- a/Engine/Plugins/Experimental/PlainProps/Source/Public/PlainPropsBuild.h
+++ b/Engine/Plugins/Experimental/PlainProps/Source/Public/PlainPropsBuild.h
@@ -62,13 +62,13 @@ public:
 		return Out;
 	}
 
-	template<typename T>
-	inline T* AllocateArray(uint64 Num)
+	template<typename T, typename... Ts>
+	inline T* AllocateArray(uint64 Num, Ts... ConstructorParams)
 	{
 		T* Out = static_cast<T*>(Allocate(Num * sizeof(T), alignof(T)));
 		for (uint64 Idx = 0; Idx < Num; ++Idx)
 		{
-			new (Out + Idx) T;
+			new (Out + Idx) T(ConstructorParams...);
 		}
 		return Out;
 	}
@@ -116,6 +116,13 @@ inline bool operator==(FMemberSchema A, FMemberSchema B)
 }
 //////////////////////////////////////////////////////////////////////////
 
+template<Arithmetic T>
+inline T Normalize(T Value)					{ return Value; }
+inline float Normalize(float Value)			{ return Value; } // todo
+inline double Normalize(double Value)		{ return Value; } // todo
+
+//////////////////////////////////////////////////////////////////////////
+
 inline uint64 ValueCast(bool Value)			{ return static_cast<uint64>(Value); }
 inline uint64 ValueCast(int8 Value)			{ return static_cast<uint8>(Value); }
 inline uint64 ValueCast(int16 Value)		{ return static_cast<uint16>(Value); }
diff --git a/Engine/Plugins/Experimental/PlainProps/Source/Public/PlainPropsDeclare.h b/Engine/Plugins/Experimental/PlainProps/Source/Public/PlainPropsDeclare.h
--- a/Engine/Plugins/Experimental/PlainProps/Source/Public/PlainPropsDeclare.h
+++ b/Engine/Plugins/Experimental/PlainProps/Source/Public/PlainPropsDeclare.h
@@ -31,6 +31,8 @@ struct FEnumerator
 {
 	FNameId							Name;
 	uint64							Constant;
+
+	friend bool operator==(FEnumerator, FEnumerator) = default;
 };
 
 enum class EEnumMode : uint8 { Flat, Flag };
diff --git a/Engine/Plugins/Experimental/PlainProps/Source/Public/PlainPropsIndex.h b/Engine/Plugins/Experimental/PlainProps/Source/Public/PlainPropsIndex.h
--- a/Engine/Plugins/Experimental/PlainProps/Source/Public/PlainPropsIndex.h
+++ b/Engine/Plugins/Experimental/PlainProps/Source/Public/PlainPropsIndex.h
@@ -97,10 +97,10 @@ class FLiteralIndexerBase : public FIdIndexerBase
 {
 public:
 	// Override to split scope into parts, default implementation assumes a single flat scope
-	virtual FScopeId					IndexScopeLiteral(FAnsiStringView Scope)			{ return FScopeId(FFlatScopeId{IndexLiteral(Scope)}); }
+	virtual FScopeId					IndexScopeLiteral(FAnsiStringView Scope)			{ return Scope.Len() ? FScopeId(FFlatScopeId{IndexLiteral(Scope)}) : NoId; }
 	FConcreteTypenameId					IndexTypenameLiteral(FAnsiStringView Type)			{ return {IndexLiteral(Type)}; }
 	FMemberId							IndexMemberLiteral(FAnsiStringView Member)			{ return {IndexLiteral(Member)}; }
-	FNameId								IndexEnumFlagLiteral(FAnsiStringView Flag)			{ return IndexLiteral(Flag); }
+	FNameId								IndexEnumeratorLiteral(FAnsiStringView Flag)		{ return IndexLiteral(Flag); }
 };
 
 template<class NameType>
diff --git a/Engine/Plugins/Experimental/PlainProps/Source/Public/PlainPropsRead.h b/Engine/Plugins/Experimental/PlainProps/Source/Public/PlainPropsRead.h
--- a/Engine/Plugins/Experimental/PlainProps/Source/Public/PlainPropsRead.h
+++ b/Engine/Plugins/Experimental/PlainProps/Source/Public/PlainPropsRead.h
@@ -130,6 +130,15 @@ public:
 	{
 #if DO_CHECK
 		check(End - It == ExpectedSize);
+#endif
+	}
+
+	uint64 CheckableSize() const
+	{
+#if DO_CHECK
+		return static_cast<uint64>(End - It);
+#else
+		return nullptr;
 #endif
 	}
 };
@@ -374,9 +383,9 @@ private:
 class FBoolRangeView
 {
 public:
-	FBoolRangeView(const uint8* InData, uint64 InNum)
-	: Data(InData)
-	, NumBits(InNum)
+	FBoolRangeView(const void* Bits, uint64 Num)
+	: Data(static_cast<const uint8*>(Bits))
+	, NumBits(Num)
 	{}
 
 	inline uint64				Num() const								{ return NumBits;}
diff --git a/Engine/Plugins/Experimental/PlainProps/Source/Public/PlainPropsTypes.h b/Engine/Plugins/Experimental/PlainProps/Source/Public/PlainPropsTypes.h
--- a/Engine/Plugins/Experimental/PlainProps/Source/Public/PlainPropsTypes.h
+++ b/Engine/Plugins/Experimental/PlainProps/Source/Public/PlainPropsTypes.h
@@ -279,13 +279,12 @@ struct FType
 
 //////////////////////////////////////////////////////////////////////////
 
-inline uint32 ToIdx(FInnerId Id) { return Id.Idx; }
-inline uint32 ToIdx(FNameId Id) { return Id.Idx; }
+template<class IdType>
+uint32 ToIdx(IdType Id)
+{
+	return Id.Idx;
+}
 inline uint32 ToIdx(FMemberId Name) { return Name.Id.Idx; }
-inline uint32 ToIdx(FEnumId Id) { return Id.Idx; }
-inline uint32 ToIdx(FStructId Id) { return Id.Idx; }
-inline uint32 ToIdx(FSchemaId Id) { return Id.Idx; }
-inline uint32 ToIdx(FNestedScopeId Id) { return Id.Idx; }
 inline uint32 ToIdx(FParametricTypeId Id) { return Id.AsInt(); }
 inline uint32 ToIdx(FConcreteTypenameId Name) { return Name.Id.Idx; }
 
@@ -358,15 +357,18 @@ using FOptionalConcreteTypenameId = TOptionalId<FConcreteTypenameId>;
 template<class IdType>
 inline constexpr TOptionalId<IdType> ToOptional(IdType Id) { return Id; }
 
-inline FOptionalInnerId			ToOptionalInner(FOptionalEnumId In)		{ return static_cast<const FOptionalInnerId&>(static_cast<const FOptionalId&>(In)); }
-inline FOptionalInnerId			ToOptionalInner(FOptionalStructId In)	{ return static_cast<const FOptionalInnerId&>(static_cast<const FOptionalId&>(In)); }
-inline FOptionalEnumId			ToOptionalEnum(FOptionalInnerId In)		{ return static_cast<const FOptionalEnumId&>(static_cast<const FOptionalId&>(In)); }
-inline FOptionalStructId		ToOptionalStruct(FOptionalInnerId In)	{ return static_cast<const FOptionalStructId&>(static_cast<const FOptionalId&>(In)); }
-inline FOptionalStructId		ToOptionalStruct(FOptionalDeclId In)	{ return static_cast<const FOptionalStructId&>(static_cast<const FOptionalId&>(In)); }
-inline FOptionalStructId		ToOptionalStruct(FOptionalBindId In)	{ return static_cast<const FOptionalStructId&>(static_cast<const FOptionalId&>(In)); }
-inline FOptionalDeclId			ToOptionalDeclId(FOptionalInnerId In)	{ return static_cast<const FOptionalDeclId&>(static_cast<const FOptionalId&>(In)); }
-inline FOptionalEnumSchemaId	ToOptionalEnum(FOptionalSchemaId In)	{ return static_cast<const FOptionalEnumSchemaId&>(static_cast<const FOptionalId&>(In)); }
-inline FOptionalStructSchemaId	ToOptionalStruct(FOptionalSchemaId In)	{ return static_cast<const FOptionalStructSchemaId&>(static_cast<const FOptionalId&>(In)); }
+inline FOptionalInnerId			ToOptionalInner(FOptionalEnumId In)				{ return static_cast<const FOptionalInnerId&>(static_cast<const FOptionalId&>(In)); }
+inline FOptionalInnerId			ToOptionalInner(FOptionalStructId In)			{ return static_cast<const FOptionalInnerId&>(static_cast<const FOptionalId&>(In)); }
+inline FOptionalInnerId			ToOptionalInner(FOptionalDeclId In)				{ return static_cast<const FOptionalInnerId&>(static_cast<const FOptionalId&>(In)); }
+inline FOptionalEnumId			ToOptionalEnum(FOptionalInnerId In)				{ return static_cast<const FOptionalEnumId&>(static_cast<const FOptionalId&>(In)); }
+inline FOptionalStructId		ToOptionalStruct(FOptionalInnerId In)			{ return static_cast<const FOptionalStructId&>(static_cast<const FOptionalId&>(In)); }
+inline FOptionalStructId		ToOptionalStruct(FOptionalDeclId In)			{ return static_cast<const FOptionalStructId&>(static_cast<const FOptionalId&>(In)); }
+inline FOptionalStructId		ToOptionalStruct(FOptionalBindId In)			{ return static_cast<const FOptionalStructId&>(static_cast<const FOptionalId&>(In)); }
+inline FOptionalDeclId			ToOptionalDeclId(FOptionalInnerId In)			{ return static_cast<const FOptionalDeclId&>(static_cast<const FOptionalId&>(In)); }
+inline FOptionalSchemaId		ToOptionalSchema(FOptionalStructSchemaId In)	{ return static_cast<const FOptionalSchemaId&>(static_cast<const FOptionalId&>(In)); }
+inline FOptionalSchemaId		ToOptionalSchema(FOptionalEnumSchemaId In)		{ return static_cast<const FOptionalSchemaId&>(static_cast<const FOptionalId&>(In)); }
+inline FOptionalEnumSchemaId	ToOptionalEnum(FOptionalSchemaId In)			{ return static_cast<const FOptionalEnumSchemaId&>(static_cast<const FOptionalId&>(In)); }
+inline FOptionalStructSchemaId	ToOptionalStruct(FOptionalSchemaId In)			{ return static_cast<const FOptionalStructSchemaId&>(static_cast<const FOptionalId&>(In)); }
 
 //////////////////////////////////////////////////////////////////////////
 
diff --git a/Engine/Plugins/Experimental/PlainProps/Source/Tests/PlainPropsUeCoreTests.cpp b/Engine/Plugins/Experimental/PlainProps/Source/Tests/PlainPropsUeCoreTests.cpp
--- a/Engine/Plugins/Experimental/PlainProps/Source/Tests/PlainPropsUeCoreTests.cpp
+++ b/Engine/Plugins/Experimental/PlainProps/Source/Tests/PlainPropsUeCoreTests.cpp
@@ -374,7 +374,7 @@ public:
 			// Mount saved schemas as is
 			Batch = MountReadSchemas(SavedSchemas);	
 
-			TArray<FStructId> ExpectedRuntimeIds = IndexRuntimeIds(*SavedSchemas, GNames);
+			TArray<FStructId> ExpectedRuntimeIds = IndexStructRuntimeIds(*SavedSchemas, GNames);
 			check(EqualItems(ExpectedRuntimeIds, InRuntimeIds));
 		}
 		
