PlainProps fix compile error

#rnx

diff --git a/Engine/Plugins/Experimental/PlainProps/Source/Public/PlainPropsRead.h b/Engine/Plugins/Experimental/PlainProps/Source/Public/PlainPropsRead.h
--- a/Engine/Plugins/Experimental/PlainProps/Source/Public/PlainPropsRead.h
+++ b/Engine/Plugins/Experimental/PlainProps/Source/Public/PlainPropsRead.h
@@ -138,7 +138,7 @@ public:
 #if DO_CHECK
 		return static_cast<uint64>(End - It);
 #else
-		return nullptr;
+		return 0;
 #endif
 	}
 };
