PlainProps upgrade system

#rnx
#Synced-CL 47103930

diff --git a/Engine/Plugins/Experimental/PlainProps/Source/Private/PlainPropsInternalUpgrade.h b/Engine/Plugins/Experimental/PlainProps/Source/Private/PlainPropsInternalUpgrade.h
new file mode 100644
--- /dev/null
+++ b/Engine/Plugins/Experimental/PlainProps/Source/Private/PlainPropsInternalUpgrade.h
@@ -0,0 +1,513 @@
+// Copyright Epic Games, Inc. All Rights Reserved.
+
+#pragma once
+
+#include "PlainPropsUpgrade.h"
+#include "PlainPropsBuild.h"
+#include "PlainPropsRead.h"
+#include "Containers/Map.h"
+#include "Containers/Set.h"
+
+namespace PlainProps::Upgrade
+{
+
+class FEnumLoosener;
+struct FEnumTransformation;
+class FInstanceUpgrader;
+struct FMatchHelper;
+struct FOps;
+struct FStructDecomposition;
+struct FVersionUpgrades;
+
+////////////////////////////////////////////////////////////////////////////////////////////////
+
+// Per struct-indexed member id - allows TArray<T> lookup instead of TMap<FMemberId, T>
+struct FParameterId
+{
+	uint16 Idx = 0xFFFF;
+
+	bool operator==(FParameterId O) const { return Idx == O.Idx; }
+	friend uint32 GetTypeHash(FParameterId Id) { return Id.Idx; };
+};
+
+struct FEnumeratorIndexer : private TSet<FNameId>
+{
+	FEnumeratorIndexer() {}
+	FEnumeratorIndexer(TSet<FNameId>&& DeclaredOrder) : TSet<FNameId>(MoveTemp(DeclaredOrder)) {}
+
+	FEnumeratorId		Index(FNameId Member)					{ return { IntCastChecked<uint16>(Add(Member).AsInteger()) }; }
+	FNameId				Resolve(FEnumeratorId Parameter) const	{ return Get(FSetElementId::FromInteger(Parameter.Idx)); }
+};
+
+using FParameterIdSet = TSet<FOptionalMemberId>;
+struct FParameterIndexer : private FParameterIdSet
+{
+	FParameterIndexer() {}
+	FParameterIndexer(FParameterIdSet&& DeclaredOrder) : FParameterIdSet(MoveTemp(DeclaredOrder)){}
+
+	FParameterId		Index(FMemberId Member)					{ return Index(ToOptional(Member)); }
+	FParameterId		Index(FOptionalMemberId Member)			{ return { IntCastChecked<uint16>(Add(Member).AsInteger()) }; }
+	FParameterId		IndexSuper()							{ return { IntCastChecked<uint16>(Add(NoId).AsInteger()) }; }
+	
+	//FParameterId		Reindex(FOptionalMemberId Member) const;
+	FOptionalMemberId	Resolve(FParameterId Parameter) const	{ return Get(FSetElementId::FromInteger(Parameter.Idx)); }	
+	uint32				Num() const								{ return uint32(TSet<FOptionalMemberId>::Num()); }
+};
+
+class FLooseIndexers 
+{
+public:
+	FLooseIndexers(const IDeclarations& Types) : Declarations(Types) {}
+
+	FEnumeratorIndexer&			Get(FEnumId Id);
+	FParameterIndexer&			Get(FDeclId Id);
+	const FEnumeratorIndexer&	Get(FEnumId Id) const									{ return const_cast<FLooseIndexers&>(*this).Get(Id); }
+	const FParameterIndexer&	Get(FDeclId Id) const									{ return const_cast<FLooseIndexers&>(*this).Get(Id); }
+//	void						Drop(FEnumId Id)										{ Enums.Remove(Id); }
+//	void						Drop(FDeclId Id)										{ Structs.Remove(Id); }
+
+private:
+	TMap<FDeclId, FParameterIndexer>	Structs;
+	TMap<FEnumId, FEnumeratorIndexer>	Enums;
+	const IDeclarations&				Declarations;
+};
+
+////////////////////////////////////////////////////////////////////////////////////////////////
+
+struct FLooseParameter
+{
+	FParameterId						Name;
+	FLooseType							Type;
+};
+
+////////////////////////////////////////////////////////////////////////////////////////////////
+
+// ~FLooseParameter with out parameter specific metadata
+struct FOutParameter
+{
+	FParameterId						Name;
+	uint8								NumInners = 0;	// Struct members or enum enumerators
+	uint8								RecompIdx = 0;	// Only for enums with NumInners > 0
+	FLooseType							Type;
+};
+
+////////////////////////////////////////////////////////////////////////////////////////////////
+
+struct FLooseSchema
+{
+	uint16 								Version = 0;
+	uint16								NumMembers = 0;
+	FLooseParameter						Members[0]; // Including super
+
+	TConstArrayView<FLooseParameter>	GetMembers() const { return MakeArrayView(Members, NumMembers); }
+};
+
+//////////////////////////////////////////////////////////////////////////////////////////////////
+
+// FLooseMetadata::IsSet, IsFullyUpgraded and !!InnermostSchema
+enum class ELooseState : uint8 { Unset, Source, Loose, UpgradedLoose };
+
+// Struct-specific loose member
+struct FStructMember
+{
+	FDeclId					Id;
+	FOptionalStructSchemaId	Schema;	// @invariant !!Schema == (State == ELooseState::Source)
+	FParameterId			Name;
+	uint16					Version;
+	ELooseState				State;
+	uint64					Num;	// Num loose members or FByteReader::CheckableBytes (0 in shipping config)
+	FStructData				Data;
+};
+
+// Enum-specific loose member
+struct FEnumMember
+{
+	FEnumId					Id;
+	FOptionalEnumSchemaId	Schema;
+	FParameterId			Name;
+	uint16					Version;
+	ELooseState				State;
+	uint8					FlagMode : 1;
+	uint16					Num;
+	FEnumData				Data;
+};
+
+// Range-specific loose member
+struct FRangeMember
+{
+	FLooseType				Type;
+	FLooseMetadata			Meta;
+	FRangeData				Data;
+};
+
+////////////////////////////////////////////////////////////////////////////////////////////////
+
+// Nameless loosened struct
+struct FLooseStruct
+{
+	FDeclId					Id;
+	uint16					Version;
+	uint16					Num;
+	FLooseMember*			Members;
+};
+
+////////////////////////////////////////////////////////////////////////////////////////////////
+
+class FEnumLooseners
+{
+	const FEnumTransformation*		Data; 
+	uint32							NumStructs;
+	uint32							NumSchemas;
+
+public:
+	FEnumLooseners(TConstArrayView<FEnumTransformation> Xforms, uint32 InNumStructs);
+	const FEnumLoosener&			operator[](FEnumSchemaId Id) const;
+};
+
+////////////////////////////////////////////////////////////////////////////////////////////////
+
+// Source struct reader producing loose members
+class FLooseMemberReader : protected FMemberReader
+{
+public:
+	FLooseMemberReader(const FLooseSchema& LooseSchema, FStructView Source, FEnumLooseners InLooseners, FScratchAllocator& InScratch);
+
+	using FMemberReader::			HasMore;
+	FLooseParameter					Peek() const;		//@pre HasMore()
+	FLooseValue						Grab();				//@pre HasMore()
+
+private:
+	const FLooseParameter*			Parameters;
+	FScratchAllocator&				Scratch;
+	FEnumLooseners					Looseners;
+
+	FRangeMember					ToLooseRange(FParameterId Name, FLooseType Type, FRangeView Range);
+	const FLooseNestedRange*		LoosenRanges(FNestedRangeView Range);
+};
+
+////////////////////////////////////////////////////////////////////////////////////////////////
+
+class FOptionalVersion
+{
+	uint16 No = 0xFFFF;
+public:
+	FOptionalVersion() = default;
+	FOptionalVersion(FNullOpt) {}
+	FOptionalVersion(uint16 Version) : No(Version) { check(*this); }
+
+	explicit operator	bool() const	{ return No != 0xFFFF; }
+	uint16				Get() const		{ check(*this); return No; }
+};
+
+////////////////////////////////////////////////////////////////////////////////////////////////
+
+// Small index range, max 255 items of a max 16.7M item range
+struct FMiniSlice
+{
+	FMiniSlice(uint32 I, uint32 N) : Idx(I), Num(N) { check(Idx == I && Num == N); }
+
+	uint32 		Idx : 24;
+	uint32		Num : 8;
+};
+
+// Compiled op and metadata
+struct FOp
+{
+	Op 			Function;
+	FMiniSlice	Outputs;
+	uint8		NumOptionalInputs;
+	uint8		NumRequiredInputs;
+
+	int32		NumInputs() const { return NumOptionalInputs + NumRequiredInputs; }
+};
+
+union FMatchIdx
+{
+	uint16		Op;
+	uint16		Decomp;
+};
+
+struct FMemberMatch
+{
+	FLooseType	Type;
+	FMatchIdx	Idx;
+	uint8		InputIdx;
+	uint8		Recipient;	// 0 = this, 1 = Outer, 2 = Outer->Outer, ...
+
+	bool		Decomposes() const { return Recipient == DecomposeRecipient; }
+
+	static constexpr uint8 DecomposeRecipient = 0xFF;
+};
+
+// Compiled upgrades for all versions of a particular struct
+struct FStructUpgrades
+{
+	//~FStructUpgrades(); // todo
+
+	FDeclId						Id;
+	FDeclId						FinalId;		// Potentially renamed struct name
+	uint16						FinalVersion;	// Potentially renamed struct version
+	uint16						FirstVersion;	// First Relevant version
+ 	uint16						NumVersions;	// Versions are contiguous but can start at >0
+	const FVersionUpgrades*		Versions[0];	// Indexed by Version - FirstVersion
+};
+
+// Compiled upgrades for a particular struct version and set of overrides
+struct FVersionUpgrades
+{
+	//~FVersionUpgrades(); // todo: destroy decompositions
+
+	uint32						OutputsOffset;		// Using offsets at the end to access Outputs faster
+	uint32						NumNames;			// Declared and renamed members, duplicated across versions for simplicity
+	uint32						NumMatches;
+	uint16						NumOps;
+	uint16						NumDecompositions;	// Struct / enum inputs with named members / enumerators / flags
+	uint16						NumRecompositions;	// Enum outputs with named enumerators / flags
+	FOptionalVersion			NextVersion;		// Next version for this OR renamed struct
+	FMiniSlice					NameMatches[0];		// #NumNames
+	
+	template<class T>
+	const T&					At(uint32 Offset) const		{ return *reinterpret_cast<const T*>(reinterpret_cast<const uint8*>(this) + OutputsOffset); }
+
+	const FMemberMatch*			GetMatches() const			{ return AlignPtr<FMemberMatch>(NameMatches + NumNames); }
+	const FOp*					GetOps() const				{ return AlignPtr<FOp>(GetMatches() + NumMatches); }
+	const uint32*				GetDecompOffsets() const	{ return AlignPtr<uint32>(GetOps() + NumOps); }
+	const uint16*				GetRecompJumps() const		{ return AlignPtr<uint16>(GetDecompOffsets() + NumDecompositions); }
+	const FOutParameter*		GetOutputs() const			{ return &At<FOutParameter>(OutputsOffset); }
+};
+
+////////////////////////////////////////////////////////////////////////////////////////////////
+
+// Source constant -> loose value
+struct FFlatEnumLoosener
+{
+	TMap<uint64, FEnumeratorId> ConstantNames;
+};
+
+// Source flag -> loose value
+struct FFlagEnumLoosener
+{
+	TArray<FEnumeratorId, TFixedAllocator<64>> FlagNames;
+};
+
+// Helps convert source enums to loose or indicate that declaration matches source schema
+class FEnumLoosener
+{
+	static constexpr uint64 NeededBit = 1;
+	static constexpr uint64 FlagBit = 2;
+	static constexpr uint64 PtrMask = ~(FlagBit | NeededBit);
+	uint64 Ptr = 0;
+
+	FEnumLoosener(const void* Loosener, uint64 Flags) : Ptr(Flags | reinterpret_cast<uint64>(Loosener)) {}
+public:
+	FEnumLoosener() {}
+	FEnumLoosener(bool bNeeded, const FFlatEnumLoosener& Flat UE_LIFETIMEBOUND) : FEnumLoosener(&Flat, bNeeded * NeededBit) {}
+	FEnumLoosener(bool bNeeded, const FFlagEnumLoosener& Flag UE_LIFETIMEBOUND) : FEnumLoosener(&Flag, bNeeded * NeededBit | FlagBit) {}
+	
+	bool						Exists() const { return !!(Ptr); }
+	bool						Needed() const { return !!(Ptr & NeededBit); }
+	bool						IsFlag() const { return !!(Ptr & FlagBit); }
+	const FFlagEnumLoosener&	AsFlag() const { checkSlow( IsFlag()); return *reinterpret_cast<FFlagEnumLoosener*>(Ptr & PtrMask); }
+	const FFlatEnumLoosener&	AsFlat() const { checkSlow(!IsFlag()); return *reinterpret_cast<FFlatEnumLoosener*>(Ptr & PtrMask); }
+};
+
+// Loose values -> built value conversion
+struct FEnumFastener
+{
+	const FEnumDeclaration*		Decl = nullptr; // null if loose FEnumeratorId indices match declared constants/flags
+};
+
+// Cached info for upgrading an enum
+struct FEnumTransformation
+{
+	FEnumLoosener				Source;
+	//const FEnumUpgrades*		Upgrades; // todo
+	FEnumFastener				Loose;
+};
+
+// Cached info for upgrading a struct
+struct FTransformation
+{
+	const FLooseSchema*			SourceSchema = nullptr;
+	const FStructUpgrades*		Upgrades = nullptr;
+	const FLooseSchema*			TargetSchema = nullptr;
+};
+
+// Cached info for upgrading all structs and enums in a source batch
+struct FTransformations
+{
+	TArray<FTransformation>		Structs;
+	TSet<FDeclId>				StructLookup;
+	TArray<FEnumTransformation>	Enums;
+	TSet<FEnumId>				EnumLookup;
+	int32						NumSourceStructs;
+
+	const FTransformation&		operator[](FStructSchemaId Id) const	{ return Structs[Id.Idx]; }
+	const FTransformation&		operator[](FDeclId Id) const			{ return Structs[StructLookup.FindId(Id).AsInteger()]; }	// TArrayView::RangeCheck checks invalid FSetElementId
+	const FEnumTransformation&	operator[](FEnumSchemaId Id) const		{ return Enums[Id.Idx - NumSourceStructs]; }
+	const FEnumTransformation&	operator[](FEnumId Id) const			{ return Enums[EnumLookup.FindId(Id).AsInteger()]; }		// TArrayView::RangeCheck checks invalid FSetElementId
+};
+
+////////////////////////////////////////////////////////////////////////////////////////////////
+
+using FLooseSchemaPtr = TUniquePtr<FLooseSchema>;
+
+// Help access input struct members in FStructMemberNames order and upgrade/loosen them first
+struct FStructDecomposition
+{
+	FDeclId						Id;
+	uint16						Version;
+	TSet<FParameterId>			InputOrder;				// empty for declared order (iota)
+	const FStructUpgrades*		Upgrades = nullptr;
+	FLooseSchemaPtr				InputSchema;
+};
+
+// Help access input enum enumerators in FEnumMembers order
+struct FEnumDecomposition
+{
+	// const FEnumUpgrades*		Upgrades = nullptr;		// todo
+	TSet<FEnumeratorId>			InputOrder;				// empty for declared order (iota)
+};
+
+////////////////////////////////////////////////////////////////////////////////////////////////
+
+// Home for all outputs
+class FOutputStack
+{
+	FOut* const					Data;
+	int32						Num = 0;
+	FScratchAllocator&			Scratch; // Convenience
+public:
+	FOutputStack(FScratchAllocator& InScratch);
+	~FOutputStack() { check(Num == 0); }
+
+	FOuts						Push(TConstArrayView<FOutParameter> Params);
+	void						Pop(int32 InNum);
+};
+
+struct FOutputScope
+{
+	FOutputScope(FOutputStack& InStack, TConstArrayView<FOutParameter> Params);
+	~FOutputScope() { Stack.Pop(Outputs.Num()); }
+
+	FOutputStack&				Stack;
+	const FOuts					Outputs;
+};
+
+// Context for upgrading a source batch
+struct FBatchContext
+{
+	// todo: add loose transformations for redirected and rename to source transformations
+	const FTransformations		Transformations;
+	FScratchAllocator&			Scratch;
+	const FIdIndexerBase&		Ids;
+	FLooseIndexers&				LooseIds; // RenameMembers can index new names
+	const FLooseType			Obsolete;
+	const FSchemaBatchId		Batch; 
+	FOutputStack				OutputStack;
+	FContext					OpCtx;
+	FDebugIds					Debug;
+};
+
+// Home for FIn of "polyary" (>1 inputs) ops
+struct FInputDependencies
+{
+	int32						MissingRequired;
+	int32						MissingOptional;
+	FIn							Inputs[0];
+};
+
+// Owner of an instance upgrader that might overload and steal inputs
+struct FOuter
+{
+	FInstanceUpgrader*			Instance = nullptr;
+};
+
+// Upgrades a single struct instance
+class FInstanceUpgrader
+{
+	UE_NONCOPYABLE(FInstanceUpgrader);
+
+	FBatchContext&					Ctx;
+	const FStructUpgrades&			Struct;
+	const FVersionUpgrades*			Current = nullptr;
+	FOuter				 			Outer;
+	const uint32					NumNames;
+	FOptionalVersion				Limit;						// Unimplemented
+	bool							bUntypedOutcome = false;	// Obsolete types and intermediate versions lack known schema
+	uint16							FirstPartialOpIdx = 0xFFFF;	// First "polyary" (>1 inputs) op idx, Flush() opt
+	uint8*							NameToScanIndices;			// #NumNames - note: wont fly all the way, should support loop around + track used matches
+	TArrayView<FInputDependencies*>	PartialInputs;				// #Max(NumOps) for versions with polyary ops
+	TArrayView<FLooseMember>		Outcome;					// #NumNames, initialized by Expect()
+	TArray<FLooseMember>			Mismatches;
+	
+	void							Expect(const FLooseSchema* Target);
+	void							Feed(FLooseMemberReader&& SourceMembers);
+	void							Feed(TConstArrayView<FLooseMember> LooseMembers);
+	void							Give(FLooseType Type, FLooseValue Value);
+	void							Give(FLooseType Type, FLooseValue Value, FMatchHelper Helper);
+	FLooseValue						Decompose(FLooseType Type, FLooseValue Value, uint16 DecompIdx);
+	FLooseValue						Decompose(FLooseType Type, FLooseValue Value, const FEnumDecomposition& Decomp);
+	FLooseValue						Decompose(FEnumMember Input, const FEnumDecomposition& Decomp);
+	FLooseValue						Decompose(FStructMember Input, const FStructDecomposition& Decomp);
+	TArrayView<FLooseMember>		Decompose(FStructView Input, const FStructDecomposition& Decomp, FParameterId Name);
+	FEnumIndexData					Reorder(FEnumeratorId Flat, const TSet<FEnumeratorId>& Order);
+	FEnumIndexData					Reorder(FEnumeratorIdSet Flags, const TSet<FEnumeratorId>& Order);
+	TArrayView<FLooseMember>		Reorder(TConstArrayView<FLooseMember> Unsorted, TConstArrayView<FLooseParameter> Params, const TSet<FParameterId>& Order);
+	void							Input(FMemberMatch Match, FLooseValue Value);
+	void							Stash(FLooseType Type, FLooseValue Value);
+	bool							Emit(FLooseType Type, FLooseValue Value);
+	void							Run(FOp Op, const FIn* Inputs, uint8 Recipient);
+	FOuter							Resolve(uint8 Recipient);
+	void							Recompose(TArrayView<FOut> Outputs, TConstArrayView<FOutParameter> Expected);
+	FEnumeratorData					Recompose(FInnermostEnumType Type, FEnumIndexData In, TConstArrayView<FEnumeratorId> Order);
+	void							Flush();
+	FLooseStruct					Finish();
+
+	FString							Print(FParameterId Parameter) const;
+	template<class T> FString		Print(T Item) const { return Ctx.Debug.Print(Item); }
+	
+public:
+	FInstanceUpgrader(FBatchContext& Ctx, const FStructUpgrades& Upgrades, uint16 Version);
+	
+	FLooseStruct					Upgrade(FStructView Source, const FLooseSchema& From, const FLooseSchema* To);
+	FLooseStruct					Upgrade(FLooseStruct In, const FLooseSchema* To);
+	FStructMember					Upgrade(FStructMember In, const FLooseSchema* To, FOuter InOuter = {});
+
+	void							Reuse(); // Todo: Optimization for upgrading range-of-structs
+};
+
+////////////////////////////////////////////////////////////////////////////////////////////////
+
+struct FSchemaCache
+{
+	const FLooseSchema*					Obtain(FDeclId Id, const IDeclarations& Types, FParameterIndexer& ParamIds);
+	TMap<FDeclId, FLooseSchemaPtr>		Cache;
+};
+
+////////////////////////////////////////////////////////////////////////////////////////////////
+
+using FOpsPtr = TUniquePtr<FOps>;
+using FStructUpgradesPtr = TUniquePtr<FStructUpgrades>;
+
+struct FHistory
+{
+	const IDeclarations&				Latest;
+	mutable FSchemaCache				LatestSchemas;
+	FLiteralIndexerBase&				Ids;
+	mutable FLooseIndexers				LooseIds;
+	FScratchAllocator					CommitScratch;
+	TArray<FOpsPtr>						CommitQueue;
+	TMap<FDeclId, FStructUpgradesPtr>	CompiledUpgrades;
+	
+	void								Compile();
+	void								Erase(FStructId Id);
+
+private:
+	void								CompileStruct(TArrayView<FOpsPtr> AllOps);
+	void								Link(TConstArrayView<FOpsPtr> Committed);
+	void								Link(FDeclId Compiled);
+	void								Link(FStructDecomposition& Dependency);
+};
+
+} // namespace PlainProps::Upgrade
\ No newline at end of file
diff --git a/Engine/Plugins/Experimental/PlainProps/Source/Private/PlainPropsUpgrade.cpp b/Engine/Plugins/Experimental/PlainProps/Source/Private/PlainPropsUpgrade.cpp
new file mode 100644
--- /dev/null
+++ b/Engine/Plugins/Experimental/PlainProps/Source/Private/PlainPropsUpgrade.cpp
@@ -0,0 +1,3215 @@
+// Copyright Epic Games, Inc. All Rights Reserved.
+
+#include "PlainPropsUpgrade.h"
+#include "PlainPropsInternalUpgrade.h"
+#include "PlainPropsIndex.h"
+#include "PlainPropsInternalBind.h"
+#include "PlainPropsInternalBuild.h"
+#include "PlainPropsInternalFormat.h"
+#include "PlainPropsInternalRead.h"
+#include "Algo/Count.h"
+#include "Algo/Find.h"
+#include "Containers/BitArray.h"
+#include <type_traits>
+
+namespace PlainProps
+{
+
+static_assert(sizeof(FLooseType) == sizeof(uint64), "AsInt() requirement");
+static_assert(std::has_unique_object_representations_v<FLooseType>, "AsInt() requirement");
+
+////////////////////////////////////////////////////////////////////////////////////////////////
+
+struct F64BitIndexIterator
+{
+	uint64					Bits = 0;
+	
+	uint8					operator*() const						{ checkSlow(Bits); return static_cast<uint8>(FMath::FloorLog2NonZero_64(Bits)); }
+	void					operator++() 							{ Bits -= (uint64(1) << operator*()); }
+	bool					operator!=(F64BitIndexIterator O) const	{ return Bits != O.Bits; }
+};
+
+// Range-based for adapter
+struct F64BitIndexRange
+{
+	uint64					Bits = 0;
+	
+	F64BitIndexIterator		begin() const	{ return {Bits}; }
+	F64BitIndexIterator		end() const		{ return {0}; }
+};
+inline F64BitIndexRange BitIndices(uint64 Bits) { return {Bits}; }
+
+//struct FFlags
+//{
+//	uint64 Bits = 0;
+//
+//	void Raise(uint8 Idx) { Bits |= (uint64(1) << Idx); }
+//	void Lower(uint8 Idx) { Bits &= ~(uint64(1) << Idx); }
+//	explicit operator bool() const { return !!Bits; }
+//};
+
+struct FEnumeratorIdSetIterator : F64BitIndexIterator
+{
+	FEnumeratorIdSetIterator(uint64 Bits) : F64BitIndexIterator{.Bits = Bits} {}
+	FEnumeratorId	operator*() const { return { .Idx = F64BitIndexIterator::operator*() }; }
+};
+static FEnumeratorIdSetIterator begin(FEnumeratorIdSet Flags)		{ return {Flags.Bits}; }
+static FEnumeratorIdSetIterator end(FEnumeratorIdSet)				{ return {0}; }
+
+static F64BitIndexIterator		begin(FEnumeratorInputSet Flags)	{ return {Flags.Bits}; }
+static F64BitIndexIterator		end(FEnumeratorInputSet)			{ return {0}; }
+
+////////////////////////////////////////////////////////////////////////////////////////////////
+
+namespace Upgrade
+{
+
+static TSet<FNameId> GetDeclaredOrder(const FEnumDeclaration& Decl)
+{
+	TSet<FNameId> Out;
+	Out.Reserve(Decl.NumEnumerators);
+	for (FEnumerator Enumerator : Decl.GetEnumerators())
+	{
+		Out.Add(Enumerator.Name);
+	}
+	return Out;
+}
+
+static FParameterIdSet GetDeclaredOrder(const FStructDeclaration& Decl)
+{
+	FParameterIdSet Out;
+	Out.Reserve(Decl.NumMembers);
+	if (Decl.Super)
+	{
+		Out.Add(NoId);
+	}
+	for (FMemberId Name : Decl.GetMemberOrder())
+	{
+		Out.Add(Name);
+	}
+	return Out;
+}
+
+template<class IndexerT, typename IdT>
+IndexerT& GetIndexer(TMap<IdT, IndexerT>& Map, IdT Id, const IDeclarations& Declarations)
+{
+	if (IndexerT* Out = Map.Find(Id))
+	{
+		return *Out;
+	}
+	else if (const auto Decl = Declarations.Find(Id))
+	{
+		return Map.Emplace(Id, GetDeclaredOrder(*Decl));
+	}
+	return Map.Emplace(Id);
+}
+
+FEnumeratorIndexer& FLooseIndexers::Get(FEnumId Id) { return GetIndexer(/* in-out */ Enums, Id, Declarations); }
+FParameterIndexer& FLooseIndexers::Get(FDeclId Id)	{ return GetIndexer(/* in-out */ Structs, Id, Declarations); }
+
+////////////////////////////////////////////////////////////////////////////////////////////////
+
+static FMemoryView ViewSkippableSlice(const void* Slice)
+{
+	uint32 NumBytesRead;
+	uint64 Size = ReadVarUInt(Slice, NumBytesRead);
+	return FMemoryView((const uint8*)Slice + NumBytesRead, Size);
+}
+
+////////////////////////////////////////////////////////////////////////////////////////////////
+
+static uint64 OutlinedHiBit = uint64(1) << 63;
+
+inline bool IsOutlined(FEnumeratorIdSet Flags)
+{
+	return Flags.Bits & OutlinedHiBit;
+}
+
+inline uint64* AsOutlined(FEnumeratorIdSet Flags)
+{
+	return reinterpret_cast<uint64*>(Flags.Bits << 1);
+}
+
+FEnumeratorIdSet NewOutlined(FScratchAllocator& Scratch, uint16 NumEnumeratorIds)
+{
+	uint64* Words = Scratch.AllocateArray<uint64>(Align(NumEnumeratorIds, 64) / 64);
+	uint64 Bits = OutlinedHiBit | (reinterpret_cast<uint64>(Words) >> 1);
+	return { Bits };
+}
+
+// todo: this will be used by TranslateOutValue
+static void SetLooseFlag(FEnumeratorIdSet& Out, FEnumeratorId Name, bool bValue)
+{
+	uint64 Mask = uint64(1) << (Name.Idx % 64);
+	uint64& Word = IsOutlined(Out) ? AsOutlined(Out)[Name.Idx/64] : Out.Bits;
+
+	if (bValue)
+	{
+		Word |= Mask;
+	}
+	else
+	{
+		Word &= ~Mask;
+	}
+}
+
+////////////////////////////////////////////////////////////////////////////////////////////////
+
+inline FParameterId GetName(FLooseMetadata Meta)
+{
+	return { static_cast<uint16>(Meta.Name) };
+}
+
+inline FParameterId GetName(const FLooseMember& Member)
+{
+	return { static_cast<uint16>(Member.GetValue().Meta.Name) };
+}
+
+inline ELooseState GetInnermostState(FLooseMetadata Meta)
+{
+	if (!Meta.IsSet)
+	{
+		return ELooseState::Unset;
+	}
+	else if (Meta.InnermostSchema)
+	{
+		return ELooseState::Source;
+	}
+
+	return Meta.IsFullyUpgraded ? ELooseState::UpgradedLoose : ELooseState::Loose;
+}
+
+inline FStructMember AsStruct(FLooseType Type, FLooseMetadata Meta, FStructData Data)
+{
+	check(Type.IsStruct());
+	FDeclId Id = Type.InnermostId.Get().AsStructDeclId();
+	FOptionalStructSchemaId Schema = ToOptionalStruct(Meta.InnermostSchema);
+	return { Id, Schema, GetName(Meta), Type.InnermostVersion, GetInnermostState(Meta), Meta.Num, Data };
+}
+
+inline FEnumMember AsEnum(FLooseType Type, FLooseMetadata Meta, FEnumData Data)
+{
+	check(Type.IsEnum());
+	FEnumId Id = Type.InnermostId.Get().AsEnum();
+	FOptionalEnumSchemaId Schema = ToOptionalEnum(Meta.InnermostSchema);
+	return { Id, Schema, GetName(Meta), Type.InnermostVersion, GetInnermostState(Meta), Type.InnermostType.Enum.FlagMode, IntCastChecked<uint16>(Meta.Num), Data };
+}
+
+inline FRangeMember AsRange(FLooseMember In)
+{
+	check(In.GetType().IsRange());
+	return { In.GetType(), In.GetValue().Meta, In.GetValue().Data.Range };
+}
+
+inline FEnumMember AsEnum(FLooseMember In)
+{
+	return AsEnum(In.GetType(), In.GetValue().Meta, In.GetValue().Data.Enum);
+}
+
+inline FStructMember AsStruct(FLooseMember In)
+{
+	return AsStruct(In.GetType(), In.GetValue().Meta, In.GetValue().Data.Struct);
+}
+
+inline FStructView AsSource(FStructSchemaId Id, FSchemaBatchId Batch, FByteReader Data)
+{
+	return { {Id, Batch}, Data };
+}
+
+inline FStructView AsSource(FStructMember In, FSchemaBatchId Batch)
+{
+	check(In.State == ELooseState::Source);
+	return { {In.Schema.Get(), Batch}, FByteReader(In.Data.Source, In.Num) };
+}
+
+static FLooseMember AsMember(FStructMember In)
+{
+	FLooseType Type = {.InnermostType = InnermostStaticStruct, .InnermostVersion = In.Version, .InnermostId = FInnerId(In.Id) };
+	FLooseMetadata Meta = { .IsSet = 1, .IsFullyUpgraded = In.State == ELooseState::UpgradedLoose, .Name = In.Name.Idx, .Num = In.Num };
+	FLooseData Data = { .Struct = In.Data };	
+	return FLooseMember(Type, {Meta, Data});
+}
+
+////////////////////////////////////////////////////////////////////////////////////////////////
+
+FInStruct FIn::Struct() const
+{
+	check(Value);
+	check(Type.IsStruct());
+	return static_cast<const FInStruct&>(*this);
+}
+
+FInStruct FIn::OptStruct() const
+{
+	check(Type.OptionalParameter);
+	check(Type.IsStruct());
+	return static_cast<const FInStruct&>(*this);
+}
+
+const FIn* FIn::AsStructMembers(uint32 N) const
+{
+	check(Type.IsStruct());
+	check(!Value.Meta.InnermostSchema);
+	check(Value.Meta.Num == N);
+	check(Value.Data.Struct.Loose);
+	return static_cast<const FIn*>(Value.Data.Struct.Loose);
+}
+
+////////////////////////////////////////////////////////////////////////////////////////////////
+
+inline const FLooseMember& ReadOutput(const FOut& Out) { return Out; }
+inline TArrayView<FLooseMember> EditOutputs(TArrayView<FOut> Outs) { return MakeArrayView(static_cast<FLooseMember*>(Outs.GetData()), Outs.Num()); }
+
+FOut::FOut(FOutParameter Param)
+: FLooseMember(Param.Type, FLooseValue{{.Name = Param.Name.Idx, .Num = Param.NumInners}, {}})
+{
+	check(Param.NumInners == 0 || Param.Type.IsEnum());
+}
+
+FOut::FOut(FOutParameter Param, FOut* Members)
+: FLooseMember(Param.Type, FLooseValue{{.IsSet = 1, .Name = Param.Name.Idx, .Num = Param.NumInners}, {.Struct = {.Loose = Members}}})
+{
+	check(Param.NumInners > 0);
+}
+
+FOutEnumFlags& FOut::EnumFlags()
+{
+	check(Type.IsEnum());
+	check(Type.InnermostType.Enum.FlagMode);
+	check(!Value.Meta.InnermostSchema);
+	check(Value.Meta.Num > 0 && Value.Meta.Num <= 64);
+	return *reinterpret_cast<FOutEnumFlags*>(this);
+}
+
+FOutStruct& FOut::Struct()
+{
+	check(Type.IsStruct());
+	return *reinterpret_cast<FOutStruct*>(this);
+}
+
+FOut* FOut::AsStructMembers(uint32 N)
+{
+	check(Type.IsStruct());
+	check(!Value.Meta.InnermostSchema);
+	check(Value.Meta.Num == N);
+	check(Value.Data.Struct.Loose);
+	return static_cast<FOut*>(Value.Data.Struct.Loose);
+}
+
+void FOutStruct::operator=(FInStruct In)
+{
+	if (Type.InnermostType.Struct.IsDynamic)
+	{
+		Type.InnermostId = In.Type.InnermostId;
+	}
+	else
+	{
+		checkf(Type.InnermostId == In.Type.InnermostId, TEXT("Set wrong struct type"));
+	}
+
+	Value.SetNamed(In.Value);
+}
+
+void FOutEnumFlag::operator=(bool bValue)
+{
+	Out.Meta.IsSet = 1;
+
+	uint64& Flags = Out.Data.Enum.Loose.Indices.Flags.Bits;
+	if (bValue)
+	{
+		Flags |= Mask;
+	}
+	else
+	{
+		Flags &= ~Mask;
+	}
+}
+
+void FOutEnumFlag::operator=(TOptional<bool> bValue)
+{
+	if (bValue)
+	{
+		operator=(*bValue);
+	}
+}
+
+////////////////////////////////////////////////////////////////////////////////////////////////
+
+FEnumLooseners::FEnumLooseners(TConstArrayView<FEnumTransformation> Xforms, uint32 InNumStructs)
+: Data(Xforms.GetData())
+, NumStructs(InNumStructs)
+, NumSchemas(InNumStructs + Xforms.Num())
+{}
+
+const FEnumLoosener& FEnumLooseners::operator[](FEnumSchemaId Id) const
+{ 
+	check(Id.Idx >= NumStructs && Id.Idx < NumSchemas);
+	return Data[Id.Idx - NumStructs].Source;
+}
+
+FEnumLooseners ViewLooseners(FBatchContext& Ctx)
+{
+	return FEnumLooseners(Ctx.Transformations.Enums, Ctx.Transformations.NumSourceStructs);
+}
+
+////////////////////////////////////////////////////////////////////////////////////////////////
+
+static FEnumeratorData Loosen(const FFlatEnumLoosener& Loosener, uint64 Constant)
+{
+	return { .Flat = Loosener.ConstantNames.FindChecked(Constant) };
+}
+
+static FEnumeratorData Loosen(const FFlagEnumLoosener& Loosener, uint64 Flags)
+{
+	FEnumeratorIdSet Out;
+	for (uint8 FlagIdx : BitIndices(Flags))
+	{
+		FEnumeratorId Name = Loosener.FlagNames[FlagIdx];
+		check(Name.Idx < 64); // todo: support >64 loose enumerators (e.g. scratch allocate bitset and maybe store num words in FLooseMeta::Num)
+		Out.Bits |= uint64(1) << Name.Idx;
+	}
+	return { .Flags = Out };
+}
+
+static uint64 Fasten(TConstArrayView<FEnumerator> Enumerators, FEnumeratorId Flat, const FEnumeratorIndexer& Indexer)
+{
+	FNameId Name = Indexer.Resolve(Flat);
+	const FEnumerator* Enumerator = Algo::FindBy(Enumerators, Name, &FEnumerator::Name);
+	check(Enumerator);
+	return Enumerator->Constant;
+}
+
+static uint64 Fasten(TConstArrayView<FEnumerator> Enumerators, FEnumeratorIdSet Flags, const FEnumeratorIndexer& Indexer)
+{
+	// opt: O(N^2). Make a FEnumFlagFastener with e.g. TMap<FEnumeratorId, uint64> or TSet<FEnumeratorId> + Decl ptr
+	uint64 Constants = 0;
+	for (FEnumeratorId Flag : Flags)
+	{
+		Constants |= Fasten(Enumerators, Flag, Indexer);
+	}
+	return Constants;
+}
+
+static uint64 Fasten(FEnumFastener Fastener, FEnumeratorId Flat, const FEnumeratorIndexer& Indexer)
+{
+	return Fastener.Decl ? Fasten(Fastener.Decl->GetEnumerators(), Flat, Indexer) : Flat.Idx;
+}
+
+static uint64 Fasten(FEnumFastener Fastener, FEnumeratorIdSet Flags, const FEnumeratorIndexer& Indexer)
+{
+	return Fastener.Decl ? Fasten(Fastener.Decl->GetEnumerators(), Flags, Indexer) : Flags.Bits;
+}
+
+////////////////////////////////////////////////////////////////////////////////////////////////
+
+static uint64 GetIntegerValue(FLeafView In)
+{
+	if (In.Leaf.Type == ELeafType::Bool)
+	{
+		return In.Value.bValue;
+	}
+
+	uint64 Out = 0;
+	FMemory::Memcpy(&Out, In.Value.Ptr, SizeOf(In.Leaf.Width));
+	return Out;
+}
+
+FLooseMemberReader::FLooseMemberReader(const FLooseSchema& Schema, FStructView Source, FEnumLooseners InLooseners, FScratchAllocator& InScratch)
+: FMemberReader(Source)
+, Parameters(Schema.Members)
+, Scratch(InScratch)
+, Looseners(InLooseners)
+{}
+
+FLooseParameter FLooseMemberReader::Peek() const
+{
+	check(HasMore());
+	return Parameters[MemberIdx];
+}
+
+FLooseValue FLooseMemberReader::Grab()
+{
+	FLooseParameter Next = Peek();
+	if (Next.Type.IsRange())
+	{
+		FRangeMember Range = ToLooseRange(Next.Name, Next.Type, FMemberReader::GrabRange());
+		return { Range.Meta, { .Range = Range.Data } };
+	}
+	else if (Next.Type.IsStruct())
+	{
+		FStructView Struct = FMemberReader::GrabStruct();
+		FLooseMetadata Meta = { .IsSet = 1, .Name = Next.Name.Idx, .InnermostSchema = Struct.Schema.Id, .Num = Struct.Values.CheckableSize() };
+		return { Meta, {.Struct = { Struct.Values.Peek() }} };
+	}
+	
+	FLeafView Leaf = GrabLeaf();
+	uint64 Value = GetIntegerValue(Leaf);
+	FLooseMetadata Meta = { .IsSet = 1, .Name = Next.Name.Idx };
+
+	if (Leaf.Enum)
+	{
+		FEnumLoosener Loosener = Looseners[Leaf.Enum.Get()];
+		if (Loosener.Needed())
+		{
+			FEnumeratorData Data = Next.Type.InnermostType.Enum.FlagMode 
+								 ? Loosen(Loosener.AsFlag(), Value)
+								 : Loosen(Loosener.AsFlat(), Value);
+			return { Meta, FLooseData{ .Enum = { .Loose = { .Enumerators = Data} } } };		
+		}
+		else // source schema matches declaration
+		{
+			Meta.InnermostSchema = ToOptionalSchema(Leaf.Enum);
+			return { Meta, FLooseData{ .Enum = {.Source = Value} } };	
+		}
+	}
+	else
+	{
+		return { Meta, FLooseData{.Arithmetic = Value} };
+	}
+}
+
+FRangeMember FLooseMemberReader::ToLooseRange(FParameterId Name, FLooseType Type, FRangeView Range)
+{
+	FOptionalSchemaId SchemaId = GetSchema(Range).InnermostSchema;
+	FLooseMetadata Meta = { .IsSet = 1, .IsSourceRange = 1, .Name = Name.Idx, .InnermostSchema = SchemaId, .Num = Range.Num() };
+	if (Type.NumRanges > 1)
+	{
+		Meta.IsSourceRange = 0;
+		return { Type, Meta, {.Ranges = LoosenRanges(Range.AsRanges())} };
+	}
+	else if (Type.InnermostType.Kind == EInnermostKind::Struct)
+	{
+		return { Type, Meta, {.Structs = {.Source = GetValues(Range).GetData()}} };
+	}
+	else
+	{
+		return { Type, Meta, {.Leaves = GetValues(Range).GetData()} };
+	}
+}
+
+const FLooseNestedRange* FLooseMemberReader::LoosenRanges(FNestedRangeView Ins)
+{
+	FLooseNestedRange* Outs = Scratch.AllocateArray<FLooseNestedRange>(Ins.Num());
+	FLooseNestedRange* OutIt = Outs;
+	EMemberKind ItemKind = Ins.GetSchema().ItemType.GetKind();
+	if (ItemKind == EMemberKind::Range)
+	{
+		for (FRangeView In : Ins)
+		{
+			*OutIt++ = { In.Num(), {.Ranges = LoosenRanges(In.AsRanges())} };
+		}
+	}
+	else
+	{
+		bool bStruct = ItemKind == EMemberKind::Struct;
+		for (FRangeView In : Ins)
+		{
+			*OutIt++ = { In.Num(), bStruct ? FRangeData{ .Structs = {.Source = GetValues(In).GetData()} }
+										   : FRangeData{ .Leaves  = GetValues(In).GetData()} };
+		}
+	}
+
+	return Outs;
+}
+
+////////////////////////////////////////////////////////////////////////////////////////////////
+
+template<class T>
+T PopFirst(TArrayView<T>& View)
+{
+	T Out = View[0];
+	View.RightChopInline(1);
+	return Out;
+}
+
+template<class T>
+TArrayView<T> PopFront(TArrayView<T>& View, int32 Num)
+{
+	TArrayView<T> Out = View.Slice(0, Num);
+	View.RightChopInline(Num);
+	return Out;
+}
+
+////////////////////////////////////////////////////////////////////////////////////////////////
+
+static constexpr uint32 OutputStackSize = 2048;
+
+FOutputStack::FOutputStack(FScratchAllocator& InScratch)
+: Data(static_cast<FOut*>(InScratch.Allocate(sizeof(FOut) * 2048, alignof(FOut))))
+, Scratch(InScratch)
+{}
+
+
+static FOut AllocateStructMembers(FOutParameter Param, TConstArrayView<FOutParameter>& InOutParams, FScratchAllocator& Scratch)
+{
+	uint32 Num = Param.NumInners;
+	FOut* Members = static_cast<FOut*>(Scratch.Allocate(sizeof(FOut) * Num, alignof(FOut)));
+	for (uint8 Idx = 0; Idx < Num; ++Idx)
+	{
+		FOutParameter Member = PopFirst(InOutParams);
+		new (Members + Idx) FOut(Member.NumInners && Param.Type.IsStruct() ? AllocateStructMembers(Member, InOutParams, Scratch) : MoveTemp(Member));
+	}
+	return FOut(Param, Members);
+}
+
+FOuts FOutputStack::Push(TConstArrayView<FOutParameter> Params)
+{
+	check(Num + Params.Num() < OutputStackSize);
+
+	// TPagedArray impl -- requires adding TPagedArray::SetNumUninitialized
+	//int32 Idx = Stack.Num();
+	//Stack.SetNumUninitialized(Idx + Params.Num());
+	//check(&Stack[Idx + Params.Num() - 1] - &Stack[Idx] == Params.Num() - 1);
+	//// Copy expected types so FOut can type-check
+	//FOut* OutIt = &Stack[Idx];
+	//for (FLooseParameter Param : Params)
+	//{
+	//	new (OutIt++) FOut(Param); 
+	//}
+
+	const int32 Idx = Num;
+	while (!Params.IsEmpty())
+	{
+		FOutParameter Param = PopFirst(/* in-out */ Params);
+		new (Data + Num) FOut(Param.NumInners && Param.Type.IsStruct() ? AllocateStructMembers(Param, /* in-out */ Params, Scratch) : FOut(Param));
+		++Num;
+	}
+	
+	return MakeArrayView(Data + Idx, Num - Idx);
+}
+
+void FOutputStack::Pop(int32 InNum)
+{
+	static_assert(std::is_trivially_destructible_v<FOut>);
+	checkSlow(Num >= InNum);
+	Num -= InNum;
+}
+
+FOutputScope::FOutputScope(FOutputStack& InStack, TConstArrayView<FOutParameter> Params)
+: Stack(InStack)
+, Outputs(Stack.Push(Params))
+{}
+
+////////////////////////////////////////////////////////////////////////////////////////////////
+
+inline TConstArrayView<const FVersionUpgrades*> ViewVersions(const FStructUpgrades& Struct)
+{
+	return MakeArrayView(Struct.Versions, Struct.NumVersions);
+}
+
+inline const FVersionUpgrades* GetVersion(const FStructUpgrades& Struct, uint16 Version)
+{
+	check(Version >= Struct.FirstVersion && Version < Struct.FirstVersion + Struct.NumVersions);
+	return Struct.Versions[Version - Struct.FirstVersion];
+}
+
+static TConstArrayView<const FVersionUpgrades*> ViewVersionsFrom(const FStructUpgrades& Struct, uint16 FromVersion)
+{
+	int32 Skip = FromVersion - Struct.FirstVersion;
+	check(Skip >= 0 && Skip <= Struct.NumVersions);
+	return MakeArrayView(Struct.Versions + Skip, Struct.NumVersions - Skip);
+}
+
+inline TConstArrayView<FMiniSlice> ViewNameMatches(const FVersionUpgrades& Version)
+{
+	return MakeArrayView(Version.NameMatches, static_cast<int32>(Version.NumNames));
+}
+
+inline TConstArrayView<FMemberMatch> ViewMatches(const FVersionUpgrades& Version)
+{
+	return MakeArrayView(Version.GetMatches(), static_cast<int32>(Version.NumMatches));
+}
+
+inline TConstArrayView<FOp> ViewOps(const FVersionUpgrades& Version)
+{
+	return MakeArrayView(Version.GetOps(), static_cast<int32>(Version.NumOps));
+}
+
+////////////////////////////////////////////////////////////////////////////////////////////////
+
+// Dependencies are reused between all FVersionUpgrades that need them, i.e. has any multi-input op
+static TArrayView<FInputDependencies*> AllocateRequiredPartialInputs(FScratchAllocator& Scratch, TConstArrayView<const FVersionUpgrades*> Versions)
+{
+	uint16 MaxOps = 0;
+	for (const FVersionUpgrades* Version : Versions)
+	{
+		for (FOp Op : ViewOps(*Version))	
+		{
+			if (Op.NumInputs() > 1)
+			{
+				MaxOps = FMath::Max(MaxOps, Version->NumOps);
+				break;
+			}
+		}
+	}
+
+	FInputDependencies** Deps = MaxOps ? Scratch.AllocateArray<FInputDependencies*>(MaxOps, nullptr) : nullptr;
+	return MakeArrayView(Deps, MaxOps);
+}
+
+FInstanceUpgrader::FInstanceUpgrader(FBatchContext& InCtx, const FStructUpgrades& InStruct, uint16 Version)
+: Ctx(InCtx)
+, Struct(InStruct)
+, Current(GetVersion(Struct, Version))
+, NumNames(InStruct.Versions[0]->NumNames)
+, NameToScanIndices(Ctx.Scratch.AllocateArray<uint8>(NumNames, 0))
+, PartialInputs(AllocateRequiredPartialInputs(Ctx.Scratch, ViewVersionsFrom(InStruct, Version)))
+{}
+
+static FLooseMember MakeMember(FLooseParameter Param)
+{
+	FLooseValue UnsetNamedValue = {{.Name = Param.Name.Idx}, {}};
+	return FLooseMember(Param.Type, UnsetNamedValue);
+}
+
+static TArrayView<FLooseMember> CreateDummyMembers(uint16 Num, FBatchContext& Ctx)
+{
+	FLooseMember Dummy(Ctx.Obsolete, {});
+	return MakeArrayView(Ctx.Scratch.AllocateArray<FLooseMember>(Num, Dummy), Num);
+}
+
+static TArrayView<FLooseMember> CreateMembersInNameOrder(TConstArrayView<FLooseParameter> Params, FBatchContext& Ctx)
+{
+	check(Params.Num() > 0);
+
+	uint16 MaxNameIdx = 0;
+	for (FLooseParameter Param : Params)
+	{
+		MaxNameIdx = FMath::Max(MaxNameIdx, Param.Name.Idx);
+	}
+	const int32 Num = MaxNameIdx + 1;
+
+	TArrayView<FLooseMember> Out = CreateDummyMembers(Num, Ctx);
+	for (FLooseParameter Param : Params)
+	{
+		check(Out[Param.Name.Idx].GetType().Compare(Ctx.Obsolete));
+		Out[Param.Name.Idx] = MakeMember(Param);
+	}
+
+	return Out;
+}
+
+void FInstanceUpgrader::Expect(const FLooseSchema* Target)
+{
+	check(!Target || Target->NumMembers <= NumNames);
+	Outcome = Target ? CreateMembersInNameOrder(Target->GetMembers(), Ctx) 
+					 : CreateDummyMembers(Ctx.LooseIds.Get(Struct.Id).Num(), Ctx);
+	bUntypedOutcome = !Target;
+}
+
+struct FMatchHelper
+{
+	explicit FMatchHelper(const FVersionUpgrades& Version)
+	: NameMatches(ViewNameMatches(Version))
+	, AllMatches(ViewMatches(Version))
+	{}
+
+	TConstArrayView<FMiniSlice>		NameMatches;
+	TConstArrayView<FMemberMatch>	AllMatches;
+};
+
+void FInstanceUpgrader::Feed(FLooseMemberReader&& It)
+{
+	const FMatchHelper Helper(*Current);
+	while (It.HasMore())
+	{
+		FLooseParameter Next = It.Peek();
+		Give(Next.Type, It.Grab(), Helper);
+	}
+}
+
+void FInstanceUpgrader::Feed(TConstArrayView<FLooseMember> Members)
+{
+	check(uint32(Members.Num()) <= NumNames);
+	const FMatchHelper Helper(*Current);
+	for (FLooseMember Member : Members)
+	{
+		if (Member)
+		{
+			Give(Member.GetType(), Member.GetValue(), Helper);
+		}
+	}
+}
+
+static const FTransformation& Lookup(const FTransformations& Xforms, FOptionalStructSchemaId Schema, FOptionalDeclId Id)
+{
+	return Schema ? Xforms[Schema.Get()] : Xforms[Id.Get()];
+}
+
+static const FTransformation* FindUpgrade(FLooseType Type, FLooseMetadata Meta, FBatchContext& Ctx)
+{
+	if (!Meta.IsSet || Meta.IsFullyUpgraded || !Type.InnermostId)
+	{
+		return nullptr;
+	}
+	else if (Type.InnermostType.Kind == EInnermostKind::Struct)
+	{
+		const FTransformation& Xform = Lookup(Ctx.Transformations, ToOptionalStruct(Meta.InnermostSchema), ToOptionalDeclId(Type.InnermostId));
+		return Xform.Upgrades ? &Xform : nullptr;
+	}
+	else
+	{
+		check(Type.InnermostType.Kind == EInnermostKind::Enum);
+		return nullptr; // todo: enum upgrades
+	}
+}
+
+FORCENOINLINE static FLooseMember UpgradeUnmatchedMember(FLooseType Type, FLooseValue Value, const FTransformation& Xform, FBatchContext& Ctx)
+{
+	checkf(Type.IsStruct(), TEXT("TODO: Upgrade outdated [nested] struct and enum ranges"));
+	FStructMember In = AsStruct(Type, Value.Meta, Value.Data.Struct);
+	FInstanceUpgrader Upgrader(Ctx, *Xform.Upgrades, Type.InnermostVersion);
+	FStructMember Out = Upgrader.Upgrade(In, Xform.TargetSchema);
+	return AsMember(Out);
+}
+
+void FInstanceUpgrader::Give(FLooseType Type, FLooseValue Value, FMatchHelper Helper)
+{
+	FParameterId Name = GetName(Value.Meta);
+	if (Name.Idx >= NumNames) // Unknown member, potentially from obsolete type
+	{
+		Stash(Type, Value);
+		return;
+	}
+
+	const FMiniSlice NameMatches = Helper.NameMatches[Name.Idx];
+	uint8& ScanIdx = NameToScanIndices[Name.Idx];
+Rematch:
+	for (FMemberMatch Match : Helper.AllMatches.Slice(NameMatches.Idx + ScanIdx, NameMatches.Num - ScanIdx))
+	{
+		++ScanIdx;
+		check(ScanIdx);
+		if (Match.Type.Compare(Type)) // todo: check code-gen and change WithoutMetadata impl if needed
+		{
+			if (Match.Decomposes())
+			{
+				Value = Decompose(Type, Value, Match.Idx.Decomp);
+
+				checkf(ScanIdx < NameMatches.Num, TEXT("Input decomposition match can't be last match for '%s' in '%s'"), *Print(Name), *Print(Struct.Id));
+				Match = Helper.AllMatches[NameMatches.Idx + ScanIdx];
+				++ScanIdx;
+				checkf(Match.Type.Compare(Type), TEXT("The match after an input decomposition match should be of same type"));
+				checkf(!Match.Decomposes(), TEXT("Can't input decompose twice"));
+			}
+
+			Input(Match, Value);
+			return;
+		}
+	}
+
+	// Unmatched
+
+	if (const FTransformation* Xform = FindUpgrade(Type, Value.Meta, Ctx))
+	{
+		// Todo: Only upgrade to next matching version, e.g. upgrading V1 -> V2 -> V3 would miss a V2 match.
+		//		 Determine next version by looking at matches
+		FLooseMember Upgraded = UpgradeUnmatchedMember(Type, Value, *Xform, Ctx);
+		checkSlow(Name.Idx == GetName(Upgraded).Idx);
+
+		bool bVersionUpdated = !Upgraded.GetType().Compare(Type);
+		Type = Upgraded.GetType();
+		Value = Upgraded.GetValue();
+
+		if (bVersionUpdated)
+		{
+			ScanIdx = 0;
+			goto Rematch;
+		}
+	}
+
+	Stash(Type, Value);
+}
+
+void FInstanceUpgrader::Give(FLooseType Type, FLooseValue Value)
+{
+	check(Current);
+	Give(Type, Value, FMatchHelper(*Current));
+}
+
+void FInstanceUpgrader::Stash(FLooseType Type, FLooseValue Value)
+{
+	if (!Emit(Type, Value))
+	{
+		Mismatches.Emplace(Type, Value);
+	}
+}
+
+bool FInstanceUpgrader::Emit(FLooseType Type, FLooseValue Value)
+{
+	FParameterId Name = GetName(Value.Meta);
+	if (Name.Idx >= static_cast<uint32>(Outcome.Num()))
+	{
+		return false;
+	}
+	else if (FLooseMember& Out = Outcome[Name.Idx]; Out)
+	{
+		return false;
+	}
+	else if (bUntypedOutcome)
+	{
+		Out = FLooseMember(Type, Value);
+		return true;
+	}
+	else if (Type.Compare(Out.GetType()))
+	{
+		Out.SetTypechecked(Value);
+		return true;
+	}
+	return false;
+}
+
+void FInstanceUpgrader::Input(FMemberMatch Match, FLooseValue Value)
+{
+	check(Match.Idx.Op < Current->NumOps);
+	FOp Op =  Current->GetOps()[Match.Idx.Op];
+	FIn In(Match.Type, Value);
+	if (Op.NumInputs() == 1)
+	{
+		Run(Op, &In, Match.Recipient); // Unary op
+	}
+	else
+	{
+		FInputDependencies*& Deps = PartialInputs[Match.Idx.Op];
+		if (Deps == nullptr)
+		{
+			FInputDependencies Header = {Op.NumRequiredInputs, Op.NumOptionalInputs};
+			uint32 Size = offsetof(FInputDependencies, Inputs) + sizeof(FIn) * Op.NumInputs();
+			Deps = new (Ctx.Scratch.Allocate(Size, alignof(FInputDependencies))) FInputDependencies(Header);
+			FMemory::Memzero(Deps->Inputs, sizeof(FIn) * Op.NumInputs());
+			FirstPartialOpIdx = FMath::Min(FirstPartialOpIdx, Match.Idx.Op);
+		}
+
+		// Set parameter
+		checkf(!Deps->Inputs[Match.InputIdx], TEXT("Parameter alredy set"));
+		Deps->Inputs[Match.InputIdx] = In;
+
+		// Decrement dependency count
+		int32& NumMissing = Match.Type.OptionalParameter ? Deps->MissingOptional : Deps->MissingRequired;
+		check(NumMissing > 0);
+		--NumMissing;
+
+		if (Deps->MissingRequired + Deps->MissingOptional == 0)
+		{
+			Run(Op, Deps->Inputs, Match.Recipient);
+
+			// Let Flush know Op has executed and prepare reuse by later FVersionUpgrades
+			Deps = nullptr;
+		}
+	}
+}
+
+static TArrayView<FLooseMember> NewMembersInInputOrder(TConstArrayView<FLooseParameter> Params, FBatchContext& Ctx)
+{
+	int32 Num = Params.Num();
+	FLooseMember* OutIt = static_cast<FLooseMember*>(Ctx.Scratch.Allocate(Num * sizeof(FLooseMember), alignof(FLooseMember)));
+	for (FLooseParameter Param : Params)
+	{
+		new (OutIt++) FLooseMember(MakeMember(Param));
+	}
+	return MakeArrayView(OutIt - Num, Num);
+}
+
+TArrayView<FLooseMember> FInstanceUpgrader::Reorder(TConstArrayView<FLooseMember> Unsorted, TConstArrayView<FLooseParameter> Params, const TSet<FParameterId>& Order)
+{
+	check(Params.Num() == Order.Num());
+
+	TArrayView<FLooseMember> Sorted = NewMembersInInputOrder(Params, Ctx);
+	for (FLooseMember Src : Unsorted)
+	{
+		if (Src)
+		{
+			FParameterId Name = GetName(Src);
+			FSetElementId Place = Order.FindId(Name);
+			check(Place.IsValidId()); // todo: emit Mismatches
+			FLooseMember& Dst = Sorted[Place.AsInteger()];
+			check(Src.GetType().Compare(Dst.GetType()));
+			Dst.SetTypechecked(Src.GetValue());
+		}
+	}
+
+	return Sorted;
+}
+
+void SetValueOrAddMismatch(FLooseMember* Out, FLooseType Type, FLooseValue Value, TArray<FLooseMember>& OutMismatches)
+{
+	check(Value.Meta.IsSet);
+	if (Out && Type.Compare(Out->GetType()))
+	{
+		Out->SetTypechecked(Value);
+	}
+	else
+	{
+		OutMismatches.Emplace(Type, Value);
+	}
+}
+
+static TArrayView<FLooseMember> Prune(TArrayView<FLooseMember> In, int32 Num, TArray<FLooseMember>& Mismatches)
+{
+	for (FLooseMember Surplus : In.RightChop(Num))
+	{
+		if (Surplus.GetValue())
+		{
+			Mismatches.Add(Surplus);
+		}
+	}
+
+	return In.Left(Num);
+}
+
+FLooseValue FInstanceUpgrader::Decompose(FLooseType Type, FLooseValue Value, uint16 DecompIdx)
+{
+	check(DecompIdx < Current->NumDecompositions);
+	uint32 Offset = Current->GetDecompOffsets()[DecompIdx];
+	const void* Ptr = ((const uint8*)Current) + Offset;
+	return Type.IsStruct()
+		? Decompose(AsStruct(Type, Value.Meta, Value.Data.Struct), *static_cast<const FStructDecomposition*>(Ptr))
+		: Decompose(AsEnum(Type, Value.Meta, Value.Data.Enum), *static_cast<const FEnumDecomposition*>(Ptr));
+}
+
+FLooseValue FInstanceUpgrader::Decompose(FLooseType Type, FLooseValue Value, const FEnumDecomposition& Decomp)
+{
+	check(Value.Meta.IsSet);
+	return Decomp.InputOrder.IsEmpty() && Value.Meta.IsFullyUpgraded ? Value : Decompose(AsEnum(Type, Value.Meta, Value.Data.Enum), Decomp);
+}
+
+FLooseValue FInstanceUpgrader::Decompose(FEnumMember In, const FEnumDecomposition& Decomp)
+{
+	FEnumeratorData Enumerators = {};
+	if (In.State == ELooseState::Source)
+	{
+		FEnumLoosener Loosener = Ctx.Transformations[In.Schema.Get()].Source;
+		Enumerators = In.FlagMode	? Loosen(Loosener.AsFlag(), In.Data.Source)
+									: Loosen(Loosener.AsFlat(), In.Data.Source);
+	}
+	else
+	{
+		Enumerators = In.Data.Loose.Enumerators;
+	}
+
+	FEnumIndexData Indices = In.FlagMode	? Reorder(Enumerators.Flags, Decomp.InputOrder) 
+											: Reorder(Enumerators.Flat, Decomp.InputOrder);
+
+	FLooseMetadata Meta = { .IsSet = 1, .Name = In.Name.Idx };
+	return { Meta, { .Enum = { .Loose = { .Indices = Indices } } } };
+}
+
+static int32 GetInputIndex(FEnumeratorId Name, const TSet<FEnumeratorId>& Order)
+{
+	FSetElementId Idx = Order.FindId(Name);
+	check(Idx.IsValidId());
+	return Idx.AsInteger();
+}
+
+FEnumIndexData FInstanceUpgrader::Reorder(FEnumeratorId Flat, const TSet<FEnumeratorId>& Order)
+{
+	return { .Flat = { .Index = static_cast<uint16>(GetInputIndex(Flat, Order)) } };
+}
+
+FEnumIndexData FInstanceUpgrader::Reorder(FEnumeratorIdSet Flags, const TSet<FEnumeratorId>& Order)
+{
+	uint64 Indices = 0;
+	for (FEnumeratorId Flag : Flags)
+	{
+		int32 OutIdx = GetInputIndex(Flag, Order);
+		check(OutIdx < 64);
+		Indices |= uint64(1) << OutIdx;
+	}
+	return { .Flags { .Bits = Indices } };
+}
+
+FLooseValue FInstanceUpgrader::Decompose(FStructMember In, const FStructDecomposition& Decomp)
+{
+	check(In.State != ELooseState::Unset);
+
+	const FLooseSchema& To = *Decomp.InputSchema;
+	TArrayView<FLooseMember> Members;
+	if (In.State == ELooseState::Source)
+	{
+		Members = Decompose(AsSource(In, Ctx.Batch), Decomp, In.Name);
+	}
+	else
+	{
+		if (In.State == ELooseState::Loose && Decomp.Upgrades)
+		{
+			FInstanceUpgrader Upgrader(Ctx, *Decomp.Upgrades, To.Version);
+			FStructMember Upgraded = Upgrader.Upgrade(In, &To);
+			Members = MakeArrayView(Upgraded.Data.Loose, static_cast<int32>(Upgraded.Num));
+		}
+		else
+		{
+			Members = MakeArrayView(In.Data.Loose, static_cast<int32>(In.Num));
+		}
+	
+		if (!Decomp.InputOrder.IsEmpty())
+		{
+			Members = Reorder(Members, To.GetMembers(), Decomp.InputOrder);
+		}
+		else if (Members.Num() > To.NumMembers) // Declaration order
+		{
+			Members = Prune(Members, To.NumMembers, /* out */ Mismatches);
+		}
+	}
+	check(Members.Num() == To.NumMembers);
+	
+	return { {.IsSet = 1, .Name = In.Name.Idx, .Num = static_cast<uint32>(Members.Num())}, {.Struct = {.Loose = Members.GetData()}} };
+}
+
+TArrayView<FLooseMember> FInstanceUpgrader::Decompose(FStructView In, const FStructDecomposition& Decomp, FParameterId Name)
+{
+	const FTransformation& Xform = Ctx.Transformations[In.Schema.Id];
+	const FLooseSchema& From = *Xform.SourceSchema;
+	const FLooseSchema& To = *Decomp.InputSchema;
+	
+	if (Decomp.Upgrades)
+	{
+		FInstanceUpgrader Upgrader(Ctx, *Decomp.Upgrades, To.Version);
+		FLooseStruct Upgraded = Upgrader.Upgrade(In, From, &To);
+		TArrayView<FLooseMember> Members(Upgraded.Members, Upgraded.Num);
+		return Decomp.InputOrder.IsEmpty() ? Members : Reorder(Members, To.GetMembers(), Decomp.InputOrder);
+	}
+
+	// Loosen source into input order
+	check(!Xform.Upgrades);
+
+	TArrayView<FLooseMember> Members = NewMembersInInputOrder(To.GetMembers(), Ctx);
+	FLooseMemberReader It(To, In, ViewLooseners(Ctx), Ctx.Scratch);
+	const TSet<FParameterId>& Order = Decomp.InputOrder;
+	if (Order.IsEmpty())
+	{
+		while (It.HasMore())
+		{
+			FLooseParameter Param = It.Peek();
+			FLooseValue Value = It.Grab();
+			SetValueOrAddMismatch(Param.Name.Idx < Members.Num() ? &Members[Param.Name.Idx] : nullptr, Param.Type, Value, /* out */ Mismatches);
+		}	
+	}
+	else
+	{
+		while (It.HasMore())
+		{
+			FLooseParameter Param = It.Peek();
+			FLooseValue Value = It.Grab();
+			FSetElementId Place = Order.FindId(Param.Name);
+			SetValueOrAddMismatch(Place.IsValidId() ? &Members[Place.AsInteger()] : nullptr, Param.Type, Value, /* out */ Mismatches);
+		}
+	}
+
+	return Members;
+}
+
+FOuter FInstanceUpgrader::Resolve(uint8 Recipient)
+{
+	FInstanceUpgrader* It = this;
+	while (--Recipient)
+	{
+		It = Outer.Instance;
+		check(It);
+	}
+	return It->Outer;
+}
+
+// Opt: avoid creating "no-op" recompositions if IsDeclarationOrder(Order)
+FEnumeratorData FInstanceUpgrader::Recompose(FInnermostEnumType Type, FEnumIndexData In, TConstArrayView<FEnumeratorId> Order)
+{
+	if (Type.FlagMode)
+	{
+		FEnumeratorIdSet Out = {};
+		for (uint8 InputIdx : In.Flags)
+		{
+			FEnumeratorId Enumerator = Order[InputIdx];
+			check(Enumerator.Idx < 64); // todo: implement large FEnumeratorIdSet support
+			Out.Bits |= uint64(1) << Enumerator.Idx;
+		}
+
+		return { .Flags = Out };
+	}
+	else
+	{
+		return { .Flat = Order[In.Flat.Index] };
+	}
+}
+
+static TConstArrayView<FEnumeratorId> GetRecompOrder(const FVersionUpgrades& Version, FOutParameter Output)
+{
+	const uint16& Jump = Version.GetRecompJumps()[Output.RecompIdx];
+	return MakeArrayView(reinterpret_cast<const FEnumeratorId*>(&Jump + Jump), Output.NumInners);
+}
+
+void FInstanceUpgrader::Recompose(TArrayView<FOut> Outputs, TConstArrayView<FOutParameter> Params)
+{
+	checkSlow(Outputs.Num() <= Params.Num());
+	for (FLooseMember& Output : EditOutputs(Outputs))
+	{
+		FOutParameter Param = PopFirst(Params);
+		
+		check(Output.GetType().CompareAll(Param.Type));
+		checkf(Output || Param.Type.OptionalParameter, TEXT("Out parameter '%s' is required but wasn't set while upgrading %s"), 
+			*Print(GetName(Output)), *Ctx.Debug.Print(Struct.Id));
+		
+		if (Param.NumInners && Param.Type.InnermostType.Kind == EInnermostKind::Enum)
+		{
+			check(Output.GetValue().Meta.Num == Param.NumInners);
+			Output.EditTypechecked().Meta.Num = 0; // Reset num enumerator indices
+			FEnumLooseData& Enum = Output.EditTypechecked().Data.Enum.Loose;
+			Enum.Enumerators = Recompose(Param.Type.InnermostType.Enum, Enum.Indices, GetRecompOrder(*Current, Param));
+		}
+	}
+}
+
+void FInstanceUpgrader::Run(FOp Op, const FIn* Inputs, uint8 Recipient)
+{
+	TConstArrayView<FOutParameter> OutParams(Current->GetOutputs() + Op.Outputs.Idx, Op.Outputs.Num);
+	FOutputScope Scope(Ctx.OutputStack, OutParams);
+
+	(*Op.Function)(MakeArrayView(Inputs, Op.NumInputs()), Scope.Outputs, Ctx.OpCtx);
+
+	Recompose(Scope.Outputs, OutParams);
+
+	if (Recipient)
+	{
+		FOuter To = Resolve(Recipient);
+		FMatchHelper Helper(*To.Instance->Current);
+		for (const FOut& Output : Scope.Outputs)
+		{
+			if (const FLooseMember& Member = ReadOutput(Output))
+			{
+				To.Instance->Give(Member.GetType(), Member.GetValue(), Helper);
+			}
+		}
+	}
+	else
+	{
+		for (const FOut& Output : Scope.Outputs)
+		{
+			if (const FLooseMember& Member = ReadOutput(Output))
+			{
+				Give(Member.GetType(), Member.GetValue());
+			}
+		}
+	}
+}
+
+static int64 DebugVersionNo(const FStructUpgrades& Struct, const FVersionUpgrades* Version)
+{
+	FVersionUpgrades const*const* Pos = Algo::Find(ViewVersions(Struct), Version);
+	return Struct.FirstVersion + (Pos - Struct.Versions);
+}
+
+void FInstanceUpgrader::Flush()
+{
+	check(Current);
+
+	if (FirstPartialOpIdx != 0xFFFF)
+	{
+		uint32 OpIdx = FirstPartialOpIdx;
+		for (FInputDependencies* Deps : PartialInputs.Slice(OpIdx, Current->NumOps - OpIdx))
+		{
+			if (Deps)
+			{
+				// todo: recoverable error handling? 
+				checkf(Deps->MissingRequired == 0, TEXT("Required input missing for op %d when upgrading version %d of %s"),
+													OpIdx, DebugVersionNo(Struct, Current), *Ctx.Debug.Print(Struct.Id));
+				check(Deps->MissingOptional > 0);
+
+				FOp Op =  Current->GetOps()[OpIdx];
+				Run(Op, Deps->Inputs, /* Recipient, inputs already given to this */ 0);
+			}
+			++OpIdx;
+		}
+
+		FirstPartialOpIdx = 0xFFFF;
+	}
+	
+	Current = nullptr;
+}
+
+static void RenameMembers(TArrayView<FLooseMember> Members, const FParameterIndexer& From, FParameterIndexer& To)
+{
+	for (FLooseMember& Member : Members)
+	{
+		FOptionalMemberId Name = From.Resolve(GetName(Member));
+		FParameterId Reindexed = To.Index(Name);
+		Member.Rename(Reindexed.Idx);
+	}
+}
+
+FLooseStruct FInstanceUpgrader::Finish()
+{
+	Flush();
+
+	if (Struct.Id != Struct.FinalId)
+	{
+		RenameMembers(Outcome, Ctx.LooseIds.Get(Struct.Id), Ctx.LooseIds.Get(Struct.FinalId));
+	}
+
+	return { Struct.FinalId, Struct.FinalVersion, static_cast<uint16>(Outcome.Num()), Outcome.GetData() };
+}
+
+FString FInstanceUpgrader::Print(FParameterId Parameter) const
+{
+	return Ctx.Debug.Print(Ctx.LooseIds.Get(Struct.Id).Resolve(Parameter));
+}
+
+FLooseStruct FInstanceUpgrader::Upgrade(FStructView In, const FLooseSchema& From, const FLooseSchema* To)
+{
+	check(Current == GetVersion(Struct, From.Version));
+	Expect(To);
+	Feed(FLooseMemberReader(From, In, ViewLooseners(Ctx), Ctx.Scratch));
+	return Finish();
+}
+
+FLooseStruct FInstanceUpgrader::Upgrade(FLooseStruct In, const FLooseSchema* To)
+{
+	check(Current == GetVersion(Struct, In.Version));
+	Expect(To);
+	Feed(MakeArrayView(In.Members, In.Num));
+	return Finish();
+}
+
+FStructMember FInstanceUpgrader::Upgrade(FStructMember In, const FLooseSchema* To, FOuter InOuter)
+{
+	check(Current == GetVersion(Struct, In.Version));
+
+	Outer = InOuter;
+	Expect(To);
+
+	switch (In.State)
+	{
+		case ELooseState::Source:
+		{
+			const FLooseSchema& SourceSchema = *Ctx.Transformations[In.Schema.Get()].SourceSchema;
+			check(SourceSchema.Version == In.Version);
+			Feed(FLooseMemberReader(SourceSchema, AsSource(In, Ctx.Batch), ViewLooseners(Ctx), Ctx.Scratch));
+			break;
+		} 
+		case ELooseState::Loose:
+			Feed(MakeArrayView(In.Data.Loose, In.Num));
+			break;
+		case ELooseState::UpgradedLoose: 
+			checkf(false, TEXT("Already fully upgraded"));
+			break;
+		default:
+			checkf(false, TEXT("Unset"));
+			break;
+	}
+	
+	FLooseStruct Out = Finish();
+	// todo: upgrade limit should return ELooseState::Loose
+	return { Out.Id, NoId, In.Name, Out.Version, ELooseState::UpgradedLoose, Out.Num, { .Loose = Out.Members } };
+}
+
+
+////////////////////////////////////////////////////////////////////////////////////////////////
+
+struct FImplicitFlags
+{
+	uint64 operator[](int32 Idx) const { return uint64(1) << Idx; } 
+};
+
+struct FImplicitConstants
+{
+	uint64 operator[](int32 Idx) const { return uint64(Idx); } 
+};
+
+template<typename T>
+static bool	MatchEnumerators(TConstArrayView<FEnumerator> Declaration, TConstArrayView<FNameId> Names, T Constants)
+{
+	if (Names.Num() > Declaration.Num())
+	{
+		return false;
+	}
+
+	for (int32 Idx = 0; Idx < Names.Num(); ++Idx)
+	{
+		FEnumerator Needed = {Names[Idx], Constants[Idx]};
+		int32 FoundIdx = Declaration.Find(Needed);
+		if (FoundIdx == INDEX_NONE)
+		{
+			return false;
+		}
+
+		// Drop first enumerator if that matched to avoid O(N^2) in normal case, can be optimized better...
+		Declaration.RightChopInline(FoundIdx == 0);
+	}
+	
+	return true;
+}
+
+// Check if a persisted enum schema matches latest declared schema
+static bool	MatchEnum(const FEnumDeclaration& Decl, const FEnumSchema& Schema)
+{
+	// Flag mode can differ if constants happen to match (unlikely)
+	
+	// Consider if we should require enum width declaration at this point
+	if (Decl.Width && Decl.Width.Get() != Schema.Width)
+	{
+		return false;
+	}
+
+	TConstArrayView<FEnumerator> Enumerators = Decl.GetEnumerators();
+	TConstArrayView<FNameId> SchemaNames(Schema.Footer, Schema.Num);
+	if (!Schema.ExplicitConstants)
+	{
+		return Schema.FlagMode	? MatchEnumerators(Enumerators, SchemaNames, FImplicitFlags{})
+								: MatchEnumerators(Enumerators, SchemaNames, FImplicitConstants{});
+	}
+	else switch (Schema.Width)
+	{
+		case ELeafWidth::B8:	return MatchEnumerators(Enumerators, SchemaNames, GetConstants<uint8>(Schema));
+		case ELeafWidth::B16:	return MatchEnumerators(Enumerators, SchemaNames, GetConstants<uint16>(Schema));
+		case ELeafWidth::B32:	return MatchEnumerators(Enumerators, SchemaNames, GetConstants<uint32>(Schema));
+		case ELeafWidth::B64:	return MatchEnumerators(Enumerators, SchemaNames, GetConstants<uint64>(Schema));
+		default:				check(false); return false;
+	}
+}
+
+////////////////////////////////////////////////////////////////////////////////////////////////
+
+// Check if a persisted struct schema matches latest declared schema
+static bool	MatchStruct(const FStructDeclaration& Decl, const FStructSchema& Schema, FInnerIds RuntimeIds)
+{
+	bool bSkipSuperSchema = SkipDeclaredSuperSchema(Schema.Inheritance);
+	uint16 NamedInnerSchemas = static_cast<uint16>(Schema.NumInnerSchemas - bSkipSuperSchema);
+
+	// Match header
+	uint16 D[4] = { Decl.Version,	Decl.NumMembers,	Decl.NumInnerIds,	Decl.NumInnerRanges };
+	uint16 S[4] = { Schema.Version, Schema.NumMembers,	NamedInnerSchemas,	Schema.NumRangeTypes };
+	if (BitCast<uint64>(D) != BitCast<uint64>(S))
+	{
+		return false;
+	}
+
+	// Match member names
+	if (!Algo::Compare(Decl.GetMemberOrder(), Schema.GetMemberNames()))
+	{
+		return false;
+	}
+
+	// Match member types. Assumes both FStructDeclaration and FStructSchema
+	// store range member types directly after member types
+	uint32 NumMemberTypes = Decl.NumMembers + Decl.NumInnerRanges;
+	if (FMemory::Memcmp(Decl.GetTypes(), Schema.Footer, NumMemberTypes * sizeof(FMemberType)) != 0)
+	{
+		return false;
+	}
+
+	// Match super
+	FOptionalDeclId SchemaSuper = Schema.GetSuper() ? ToOptional(RuntimeIds[Schema.GetSuper().Get()]) : NoId;
+	if (Decl.Super != SchemaSuper)
+	{
+		return false;
+	}
+
+	// Match inner ids
+	if (NamedInnerSchemas > 0)
+	{
+		const FInnerId* DeclIdIt = Decl.GetInnerIds();
+		const FSchemaId* SchemaIdIt = Schema.GetInnerSchemas();
+		const FSchemaId* SchemaIdEnd = SchemaIdIt + NamedInnerSchemas;
+		TConstArrayView<FMemberType> RangeTypes = Schema.GetRangeTypes();
+		uint16 RangeTypeIdx = 0;
+		for (FMemberType Member : MakeArrayView(Schema.Footer + bSkipSuperSchema, Schema.NumMembers - bSkipSuperSchema))
+		{
+			FMemberType Innermost = Member.IsRange() ? GrabInnerRangeTypes(RangeTypes, /* in-out */ RangeTypeIdx).Last() : Member;
+			if (IsStructOrEnum(Innermost))
+			{
+				FInnerId SchemaId = RuntimeIds[*SchemaIdIt++];
+				FInnerId DeclId = *DeclIdIt++;
+				if (SchemaId != DeclId)
+				{
+					return false;
+				}
+			}
+		}
+	}
+
+	return true;
+}
+
+////////////////////////////////////////////////////////////////////////////////////////////////
+
+struct FNestedMemberType
+{
+	const FMemberType*	InnerRangeTypes = nullptr;
+	uint32				NumInnerRanges = 0;
+	FMemberType			Type;
+};
+
+// CreateSourceSchema/CreateTargetSchema/FastenStruct helper
+struct FMemberTypeIterator
+{
+	const FMemberType*				TypeIt;
+	const FMemberType*				TypeEnd;
+	const FMemberType*				RangeIt;
+	const FMemberType*				RangeEnd;
+
+	FMemberTypeIterator(TConstArrayView<FMemberType> TypesWithoutSuper, TConstArrayView<FMemberType> RangeTypes)
+	: TypeIt(TypesWithoutSuper.GetData())
+	, TypeEnd(TypeIt + TypesWithoutSuper.Num())
+	, RangeIt(RangeTypes.GetData())
+	, RangeEnd(RangeIt + RangeTypes.Num())
+	{}
+
+	~FMemberTypeIterator()
+	{
+		check(TypeIt == TypeEnd);
+		check(RangeIt == RangeEnd);
+	}
+
+	FNestedMemberType Grab()
+	{
+		FNestedMemberType Out;
+		check(TypeIt < TypeEnd);
+		Out.Type = *TypeIt++;
+		Out.InnerRangeTypes = Out.Type.IsRange() ? RangeIt : nullptr;
+		for (FMemberType Type = Out.Type; Type.IsRange(); Type = *RangeIt++)
+		{
+			++Out.NumInnerRanges;
+			check(RangeIt < RangeEnd);
+		}
+		
+		return Out;
+	}
+};
+
+struct FMidMemberType
+{
+	FMidMemberType(FNestedMemberType In)
+	: NumRanges(static_cast<uint8>(In.NumInnerRanges))
+	, Innermost(In.NumInnerRanges ? In.InnerRangeTypes[In.NumInnerRanges - 1] : In.Type)
+	{
+		check(In.NumInnerRanges < MaxRangeNesting);
+	}
+
+	uint8							NumRanges;
+	FMemberType						Innermost;
+};
+
+// CreateSourceSchema helper
+struct FSchemaInnermostIdIterator
+{
+	const FSchemaId*				It;
+	FInnerIds						RuntimeIds;
+	TConstArrayView<uint16>			StructVersions;
+	const TBitArray<>&				EnumFlagModes;
+
+	FInnerId GrabStruct(uint16& OutVersion)
+	{
+		OutVersion = StructVersions[It->Idx];
+		FSchemaId Id = *It++;
+		return RuntimeIds[Id];
+	}
+
+	FInnerId GrabEnum(EEnumMode& OutMode)
+	{
+		OutMode = EnumFlagModes[It->Idx] ? EEnumMode::Flag : EEnumMode::Flat;
+		FSchemaId Id = *It++;
+		return RuntimeIds[Id];
+	}
+};
+
+// CreateTargetSchema helper
+struct FDeclaredInnermostIdIterator
+{
+	const FInnerId*					It;
+	const IDeclarations&			Declarations;
+	
+	FInnerId						Grab()
+	{
+		return *It++;
+	}
+
+	FInnerId						GrabStruct(uint16& OutVersion)
+	{
+		OutVersion = Declarations.Get(It->AsStruct()).Version;
+		return Grab();
+	}
+
+	FInnerId						GrabEnum(EEnumMode& OutMode)
+	{
+		OutMode = Declarations.Get(It->AsEnum()).Mode;
+		return Grab();
+	}
+};
+
+static FInnermostType ArithmeticToInnermost(ELeafType Type, ELeafWidth Width)
+{
+	// todo: checkSlow for illegal type and width combinations
+	static constexpr EInnermostKind SInts[4]	= {EInnermostKind::S8, EInnermostKind::S16, EInnermostKind::S32, EInnermostKind::S64};
+	static constexpr EInnermostKind UInts[4]	= {EInnermostKind::U8, EInnermostKind::U16, EInnermostKind::U32, EInnermostKind::U64};
+	static constexpr EInnermostKind Chars[4]	= {EInnermostKind::UTF8, EInnermostKind::UTF16, EInnermostKind::UTF32, EInnermostKind::UTF32};
+	static constexpr EInnermostKind Floats[4]	= {EInnermostKind::F32, EInnermostKind::F32, EInnermostKind::F32, EInnermostKind::F64};
+
+	switch (Type)
+	{
+	case ELeafType::Bool:		return { .Kind = EInnermostKind::Bool };
+	case ELeafType::IntS:		return { .Kind = SInts[(uint8)Width] };
+	case ELeafType::IntU:		return { .Kind = UInts[(uint8)Width] };
+	case ELeafType::Float:		return { .Kind = Floats[(uint8)Width] };
+	case ELeafType::Hex:		return { .Kind = UInts[(uint8)Width] };
+	case ELeafType::Unicode:	return { .Kind = Chars[(uint8)Width] };
+	default:					check(false); return {};
+	}
+}
+
+// CreateSourceSchema/CreateTargetSchema helper
+template <typename InnermostIterator>
+FInnermostType GrabInnermost(InnermostIterator& InnermostIt, FMemberType In, FOptionalInnerId& OutId, uint16& OutVersion)
+{
+	if (In.IsStruct())
+	{
+		FStructType InStruct = In.AsStruct();
+		if (InStruct.IsDynamic)
+		{
+			return { .Struct = { EInnermostKind::Struct, /* IsDynamic */ 1, InStruct.IsSuper } };	
+		}
+		OutId = ToOptional(InnermostIt.GrabStruct(OutVersion));
+		return { .Struct = { EInnermostKind::Struct, /* IsDynamic */ 0, InStruct.IsSuper } };
+	}
+
+	FUnpackedLeafType Leaf = In.AsLeaf();
+	if (Leaf.Type == ELeafType::Enum)
+	{
+		EEnumMode Mode;
+		OutId = ToOptional(InnermostIt.GrabEnum(/* out */ Mode));
+		return { .Enum = { EInnermostKind::Enum, Mode == EEnumMode::Flag } };
+	}
+
+	return ArithmeticToInnermost(Leaf.Type, Leaf.Width);
+}
+
+// CreateSourceSchema helper
+struct FSourceLooseTypeIterator
+{
+	FSourceLooseTypeIterator(const FStructSchema& Struct, FInnerIds RuntimeIds, TConstArrayView<uint16> Versions, const TBitArray<>& FlagModes)
+	: SkipSuper(SkipDeclaredSuperSchema(Struct.Inheritance))
+	, TypeIt(Struct.GetMemberTypes().LeftChop(SkipSuper), Struct.GetRangeTypes())
+	, InnerIdIt({Struct.GetInnerSchemas() + SkipSuper, RuntimeIds, Versions, FlagModes})
+	{}
+
+	const uint8					SkipSuper;
+	FMemberTypeIterator			TypeIt;
+	FSchemaInnermostIdIterator	InnerIdIt;
+
+	FLooseType					Grab()
+	{
+		FMidMemberType Mid = TypeIt.Grab();
+		FLooseType Out = { .NumRanges = Mid.NumRanges };
+		Out.InnermostType = GrabInnermost(/* in-out */ InnerIdIt, Mid.Innermost, Out.InnermostId, Out.InnermostVersion);
+		return Out;
+	}
+};
+
+// CreateTargetSchema/FastenStruct helper
+struct FTargetLooseTypeIterator
+{
+	FTargetLooseTypeIterator(const FStructDeclaration& Struct, const IDeclarations& Types)
+	: TypeIt(MakeArrayView(Struct.GetTypes(), Struct.NumMembers), 
+			 MakeArrayView(Struct.GetInnerRangeTypes(), Struct.NumInnerRanges))
+	, InnerIdIt({Struct.GetInnerIds(), Types})
+	{}
+
+	FMemberTypeIterator				TypeIt;
+	FDeclaredInnermostIdIterator	InnerIdIt;
+
+	FLooseType						GrabLoose()
+	{
+		FMidMemberType Mid = TypeIt.Grab();
+		FLooseType Out = { .NumRanges = Mid.NumRanges };
+		Out.InnermostType = GrabInnermost(/* in-out */ InnerIdIt, Mid.Innermost, Out.InnermostId, Out.InnermostVersion);
+		return Out;
+	}
+
+	FMemberSchema					GrabSchema()
+	{
+		FNestedMemberType Mid = TypeIt.Grab();
+		FMemberType InnerRangeType = Mid.InnerRangeTypes ? Mid.InnerRangeTypes[0] : Mid.Type;
+		FOptionalInnerId InnermostId = IsStructOrEnum(Mid.Type) ? ToOptional(InnerIdIt.Grab()) : NoId;
+		return { Mid.Type, InnerRangeType, static_cast<uint16>(Mid.NumInnerRanges), InnermostId, Mid.InnerRangeTypes };
+	}
+};
+
+// FBatchUpgrader::Upgrade helper
+struct FBatchRun
+{
+	FBatchRun(const FHistory& Upgrades, FScratchAllocator& InScratch, TConstArrayView<FInnerId> InnerIds)
+	: RuntimeIds{{InnerIds}}
+	, History(Upgrades)
+	, Scratch(InScratch)
+	, Debug(Upgrades.Ids)
+	{}
+
+	FInnerIds						RuntimeIds;
+	FTransformations				Transformations;
+	const FHistory&					History;
+	FScratchAllocator&				Scratch;
+	FSchemaBatchId					SourceBatch;
+	const FDebugIds					Debug;
+
+	FLooseSchema*					CreateSourceSchema(const FStructSchema& Struct, FParameterIndexer& ParamIds, TConstArrayView<uint16> StructVersions, const TBitArray<>& EnumFlagModes);
+	void							CreateTransformations(const TBitArray<>& MatchingSchemas, const FSchemaBatch& Batch);
+	FUpgradedInstance				UpgradeOrRebuild(FStructView Source);
+	FUpgradedInstance				UpgradeStruct(FStructView Source, FTransformation Xform) const;
+	FBuiltStruct*					RebuildStruct(FStructView Source) const;
+	FBuiltStruct*					RebuildStruct(FStructView Source, FDeclId Id) const;
+	void							RebuildMembers(FMemberBuilder& Out, FStructView Source, FDeclId Id) const;
+	FTypedValue						RebuildValue(FMemberReader& It) const;
+	FTypedValue						RebuildValue(FLeafView Leaf) const;
+	FTypedValue						RebuildValue(FMemberType Type, FRangeView Range) const;
+	FTypedValue						RebuildValue(FStructView Struct) const;
+	FBuiltRange*					RebuildRange(FRangeView Range) const;
+	FBuiltRange*					RebuildLeaves(FLeafRangeView Leaves) const;
+	FBuiltRange*					RebuildRanges(FNestedRangeView Ranges) const;
+	FBuiltRange*					RebuildStructs(FStructRangeView Structs) const;
+	FBuiltRange*					RebuildStructs(FStructRangeView Structs, FDeclId Id) const;
+	FBuiltStruct*					FastenStruct(FLooseStruct Loose) const;
+	void							FastenMembers(FMemberBuilder& Out, TConstArrayView<FLooseMember> Members, const FParameterIndexer& Names, const FStructDeclaration& Declaration) const;
+	FBuiltValue						FastenMember(FLooseMember Member) const;
+	uint64							FastenEnum(FEnumMember Enum) const;
+	FBuiltStruct*					FastenStruct(FStructMember Struct) const;
+	FBuiltRange*					FastenRange(FRangeMember Range) const;
+	FBuiltRange*					FastenRange(FLooseType Type, FLooseMetadata Meta, FRangeData Data) const;
+	FBuiltRange*					FastenRanges(FLooseType Type, FLooseMetadata Meta, TConstArrayView64<FLooseNestedRange> Ranges) const;
+	FBuiltRange*					FastenStructs(FDeclId Id, FLooseMetadata Meta, FStructRange Structs) const;
+	FBuiltRange*					FastenEnums(FEnumId Id, FInnermostEnumType Type, FLooseMetadata Meta, FRangeData Data) const;
+	FBuiltRange*					FastenArithmetics(EInnermostKind Kind, uint64 Num, const void* Data) const;
+	FBuiltRange*					FastenLeaves(uint64 Num, FMemoryView Data) const;
+};
+
+static FLooseSchema* CreateSchema(FScratchAllocator& Scratch, uint16 Version, uint16 NumMembers)
+{
+	const uint32 Size = offsetof(FLooseSchema, Members) + NumMembers * sizeof(FLooseParameter);
+	FLooseSchema Header = {.Version = Version, .NumMembers = NumMembers};
+	return new (Scratch.Allocate(Size, alignof(FLooseSchema))) FLooseSchema(Header);
+}
+
+FLooseSchema* FBatchRun::CreateSourceSchema(const FStructSchema& Schema, FParameterIndexer& ParamIds, TConstArrayView<uint16> StructVersions, const TBitArray<>& EnumFlagModes)
+{
+	FLooseSchema* Out = CreateSchema(Scratch, Schema.Version, Schema.NumMembers);
+
+	const FMemberId* NameIt = Schema.GetMemberNames().GetData();
+	FSourceLooseTypeIterator TypeIt(Schema, RuntimeIds, StructVersions, EnumFlagModes);
+	for (FLooseParameter& Param : MakeArrayView(Out->Members, Schema.NumNames()))
+	{
+		Param.Name = ParamIds.Index(*NameIt++);
+		Param.Type = TypeIt.Grab();
+	}
+
+	return Out;
+}
+
+static FLooseSchema* CreateTargetSchema(const FStructDeclaration& Struct, const IDeclarations& Types, FParameterIndexer& ParamIds)
+{
+	uint16 Num = Struct.NumMembers;
+	const uint32 Size = offsetof(FLooseSchema, Members) + Num * sizeof(FLooseParameter);
+	FLooseSchema Header = {.Version = Struct.Version, .NumMembers = Num};
+	FLooseSchema* Out = new (FMemory::Malloc(Size, alignof(FLooseSchema))) FLooseSchema(Header);
+
+	TArrayView<FLooseParameter> Parameters(Out->Members, Num);
+	if (FOptionalDeclId Super = Struct.Super)
+	{
+		uint16 SuperVersion = Types.Get(Super.Get()).Version;
+		Parameters[0].Name = ParamIds.IndexSuper();
+		// todo: remove or set FInnermostStructType::IsSuper
+		Parameters[0].Type = FLooseType{ 0, 0, InnermostStaticStruct, SuperVersion, ToOptionalInner(Super) };
+		Parameters.RightChopInline(1);
+	}
+
+	const FMemberId* NameIt = Struct.GetMemberOrder().GetData();
+	FTargetLooseTypeIterator TypeIt(Struct, Types);
+	for (FLooseParameter& Parameter : Parameters)
+	{
+		Parameter.Name = ParamIds.Index(*NameIt++);
+		Parameter.Type = TypeIt.GrabLoose();
+	}
+
+	return Out;
+}
+
+static TArray<uint16> ExtractVersions(TSchemaRange<const FStructSchema> Schemas)
+{
+	TArray<uint16> Out;
+	Out.Reserve(Schemas.Num());
+	for (const FStructSchema& Schema : Schemas)
+	{
+		Out.Add(Schema.Version);
+	}
+	return Out;
+}
+
+static TBitArray<> ExtractFlagModes(TSchemaRange<const FEnumSchema> Schemas)
+{
+	TBitArray<> Out;
+	Out.Reserve(Schemas.Num());
+	for (const FEnumSchema& Schema : Schemas)
+	{
+		Out.Add(Schema.FlagMode);
+	}
+	return Out;
+}
+
+static void AddInnerIdOutputs(TSet<FDeclId>& OutStructs, TSet<FEnumId>& OutEnums, const FVersionUpgrades& Version)
+{
+	const FOutParameter* Outputs = Version.GetOutputs();
+	for (FOp Op : MakeArrayView(Version.GetOps(), Version.NumOps))
+	{
+		for (FOutParameter Output : MakeArrayView(Outputs + Op.Outputs.Idx, Op.Outputs.Num))
+		{
+			if (Output.Type.InnermostId)
+			{
+				FInnerId Inner = Output.Type.InnermostId.Get();
+				if (Output.Type.IsStruct())
+				{
+					OutStructs.Emplace(Inner.AsStruct());
+				}
+				else
+				{
+					OutEnums.Add(Inner.AsEnum());
+				}
+			}
+		}
+	}
+}
+
+static FEnumLoosener CreateEnumLoosener(const FEnumSchema& Schema, bool bNeeded, FEnumeratorIndexer& Ids, FScratchAllocator& Scratch)
+{
+	unimplemented();
+	return {};
+}
+
+// Contiguous constants will match contiguous declaration order of FEnumeratorId indices
+static bool RequiresFastening(TConstArrayView<FEnumerator> Enumerators, EEnumMode Mode)
+{
+	if (Mode == EEnumMode::Flat)
+	{
+		for (int32 Idx = 0; Idx < Enumerators.Num(); ++Idx)
+		{
+			if (Enumerators[Idx].Constant != uint32(Idx))
+			{
+				return true;
+			}
+		}
+	}
+	else
+	{
+		for (int32 Idx = 0; Idx < Enumerators.Num(); ++Idx)
+		{
+			if (Enumerators[Idx].Constant != (uint64(1) << Idx))
+			{
+				return true;
+			}
+		}
+	}
+
+	return false;
+}
+
+static FEnumFastener CreateEnumFastener(const FEnumDeclaration* Decl, FScratchAllocator& Scratch)
+{
+	return { Decl && RequiresFastening(Decl->GetEnumerators(), Decl->Mode) ? Decl : nullptr };
+}
+
+void FBatchRun::CreateTransformations(const TBitArray<>& MatchingSchemas, const FSchemaBatch& Batch)
+{
+	FLooseIndexers& LooseIds = History.LooseIds;
+	const IDeclarations& Declarations = History.Latest;
+	const TMap<FDeclId, FStructUpgradesPtr>& CompiledUpgrades = History.CompiledUpgrades;
+	const TSchemaRange<const FStructSchema> StructSchemas = GetStructSchemas(Batch);
+	const TSchemaRange<const FEnumSchema> EnumSchemas = GetEnumSchemas(Batch);
+	const TArray<uint16> StructVersions = ExtractVersions(StructSchemas);
+	const TBitArray<> EnumFlagModes = ExtractFlagModes(EnumSchemas);
+	const uint32 NumSourceStructs = StructVersions.Num();
+
+	TArray<FTransformation>& StructXforms = Transformations.Structs;
+	TSet<FDeclId>& StructLookup = Transformations.StructLookup;
+	TArray<FEnumTransformation>& EnumXforms = Transformations.Enums;
+	TSet<FEnumId>& EnumLookup = Transformations.EnumLookup;
+	Transformations.NumSourceStructs = StructSchemas.Num();
+	StructXforms.SetNum(StructSchemas.Num());
+	StructLookup.Reserve(StructSchemas.Num());
+	EnumXforms.SetNum(EnumSchemas.Num());
+	EnumXforms.Reserve(EnumSchemas.Num());
+	
+	TSet<FEnumId> EnumDecompositions;
+
+	// Create initial struct transforms
+	for (uint32 Idx = 0; int32(Idx) < StructSchemas.Num(); ++Idx)
+	{
+		FDeclId Id = RuntimeIds[FStructSchemaId{Idx}];
+		if (!MatchingSchemas[Idx])
+		{
+			FParameterIndexer& Ids = LooseIds.Get(Id);
+			FTransformation& Xform = StructXforms[Idx];
+			Xform.SourceSchema = CreateSourceSchema(StructSchemas[Idx], Ids, StructVersions, EnumFlagModes);
+			Xform.Upgrades = CompiledUpgrades.FindChecked(Id).Get();
+			Xform.TargetSchema = History.LatestSchemas.Obtain(Id, Declarations, Ids);
+		}
+
+		// Ensure lookup add order follow source schema id index
+		StructLookup.Add(Id);
+	}
+
+	// Create initial enum transforms
+	for (uint32 Idx = StructSchemas.Num(); const FEnumSchema& Schema : EnumSchemas)
+	{
+		FEnumId Id = RuntimeIds[FEnumSchemaId{Idx}];
+		if (!MatchingSchemas[Idx])
+		{
+			FEnumeratorIndexer& Ids = LooseIds.Get(Id);
+			FEnumTransformation& Xform = EnumXforms[EnumLookup.Num()];
+			Xform.Source = CreateEnumLoosener(Schema, /* bNeeded */ true, Ids, Scratch);
+			Xform.Loose = CreateEnumFastener(Declarations.Find(Id), Scratch);
+		}
+	
+		// Ensure lookup add order follow source schema id index
+		EnumLookup.Add(Id);
+		++Idx;
+	}
+	
+	// Find renames and inner types outputted by upgrade ops,
+	// which aren't part of schema batch
+	TArrayView<FTransformation> Unscanned = StructXforms;
+	while (Unscanned.Num() > 0)
+	{
+		for (const FTransformation& Xform : Unscanned)
+		{
+			if (const FStructUpgrades* Upgrades = Xform.Upgrades)
+			{
+				// todo: opt - no need to scan older versions
+				// uint16 FirstVersion = Xform.SourceSchema->Version;
+				// for (const FVersionUpgrades* Version : ViewVersionsFrom(*Upgrades, FirstVersion))
+				for (const FVersionUpgrades* Version : ViewVersions(*Upgrades))
+				{
+					AddInnerIdOutputs(/* out */ StructLookup, /* out */ EnumLookup, *Version);
+				}
+
+				if (Upgrades->Id != Upgrades->FinalId)
+				{
+					StructLookup.Add(Upgrades->FinalId);
+				}
+			}
+		}
+
+		// Create FTransformations and target schemas for missing structs,
+		// regardless if upgrades exist so operator[](FDeclId) can return a reference
+		int32 Idx = StructXforms.Num();
+		StructXforms.SetNum(StructLookup.Num());
+		Unscanned = MakeArrayView(StructXforms).RightChop(Idx);
+		for (FTransformation& Xform : Unscanned)
+		{
+			FDeclId Id = StructLookup.Get(FSetElementId::FromInteger(Idx));			
+			const FStructUpgradesPtr* Upgrades = CompiledUpgrades.Find(Id);
+			Xform.Upgrades = Upgrades ? Upgrades->Get() : nullptr;
+			Xform.TargetSchema = History.LatestSchemas.Obtain(Id, Declarations, LooseIds.Get(Id));
+			++Idx;
+		}
+	}
+
+	// Create source looseners for enum input decompositions (that might not be needed)
+	for (FEnumId Id : EnumDecompositions)
+	{
+		int32 Idx = EnumLookup.FindId(Id).AsInteger();
+		if (Idx > INDEX_NONE && Idx < EnumXforms.Num())
+		{
+			FEnumLoosener& Source = EnumXforms[Idx].Source;
+			Source = Source.Exists() ? Source : CreateEnumLoosener(EnumSchemas[Idx], /* bNeeded */ false, LooseIds.Get(Id), Scratch);
+		}
+	}
+
+	// Create fasteners for enum output recompositions
+	EnumXforms.SetNum(EnumLookup.Num());
+	TArrayView<FEnumTransformation> OutputEnumXforms = MakeArrayView(EnumXforms).RightChop(EnumSchemas.Num());
+	for (int32 Idx = EnumSchemas.Num(); FEnumTransformation& Xform : OutputEnumXforms)
+	{
+		FEnumId Id = EnumLookup.Get(FSetElementId::FromInteger(Idx));
+		FEnumeratorIndexer& Ids = LooseIds.Get(Id);
+		Xform.Loose = CreateEnumFastener(Declarations.Find(Id), Scratch);
+		++Idx;
+		// todo: cache fasteners in History.LatestSchemas?
+	}
+}
+
+static FLooseType IndexObsolete(FLiteralIndexerBase& Ids)
+{
+	FStructId Obsolete = Ids.IndexStruct(FType(NoId, FTypenameId(Ids.IndexTypenameLiteral("__obsolete__"))));
+	return { .InnermostType = InnermostStaticStruct, .InnermostVersion = 0xBAAD, .InnermostId = FInnerId(Obsolete) };
+}
+
+static FLooseStruct ReorderAsDeclaration(FLooseStruct In, const FLooseSchema& Declared, FScratchAllocator& Scratch)
+{
+	// opt: check if In members have consecutive names, if so type-check and return In
+
+	FLooseMember* OutIt = static_cast<FLooseMember*>(Scratch.AllocateZeroed(Declared.NumMembers * sizeof(FLooseMember), alignof(FLooseMember)));
+	for (FLooseParameter Param : Declared.GetMembers())
+	{
+		new (OutIt++) FLooseMember(MakeMember(Param));
+	}
+
+	TArrayView<FLooseMember> Out(OutIt - Declared.NumMembers, Declared.NumMembers);
+	for (FLooseMember Src : MakeArrayView(In.Members, In.Num))
+	{
+		if (Src)
+		{
+			FParameterId Name = GetName(Src);
+			check(Name.Idx < Out.Num()); // todo: emit Mismatches
+			FLooseMember& Dst = Out[Name.Idx];
+			check(!Dst);
+			check(Src.GetType().Compare(Dst.GetType()));
+			Dst.SetTypechecked(Src.GetValue());
+		}
+	}
+
+	return { In.Id, In.Version, static_cast<uint16>(Out.Num()), Out.GetData() };
+}
+
+FUpgradedInstance FBatchRun::UpgradeStruct(FStructView Source, FTransformation Xform) const
+{
+	FBatchContext Ctx = { Transformations, Scratch, History.Ids, History.LooseIds, IndexObsolete(History.Ids), Source.Schema.Batch, {Scratch}, {}, Debug };
+	FInstanceUpgrader SourceUpgrader(Ctx, *Xform.Upgrades, Xform.SourceSchema->Version);
+	FLooseStruct Upgraded = SourceUpgrader.Upgrade(Source, *Xform.SourceSchema, Xform.TargetSchema);
+
+	// Handle renamed structs
+	while (Upgraded.Id != Xform.Upgrades->Id)
+	{
+		Xform = Transformations[Upgraded.Id];
+		if (Xform.Upgrades)
+		{
+			FInstanceUpgrader Upgrader(Ctx, *Xform.Upgrades, Upgraded.Version);
+			Upgraded = Upgrader.Upgrade(Upgraded, Xform.TargetSchema);
+		}
+		else
+		{
+			Upgraded = ReorderAsDeclaration(Upgraded, *Xform.TargetSchema, Ctx.Scratch);
+		}
+	}
+
+	return { Upgraded.Id, FastenStruct(Upgraded) };
+}
+
+FUpgradedInstance FBatchRun::UpgradeOrRebuild(FStructView Source)
+{
+	SourceBatch = Source.Schema.Batch;
+	FStructSchemaId SchemaId = Source.Schema.Id;
+	FDeclId Id = RuntimeIds[SchemaId];
+	FTransformation Xform = Transformations[Source.Schema.Id];
+	return Xform.Upgrades ? UpgradeStruct(Source, Xform) : FUpgradedInstance{ Id, RebuildStruct(Source, Id) };
+}
+
+static FUnpackedLeafType ToLeafType(EInnermostKind Arithmetic)
+{
+	switch (Arithmetic)
+	{
+	case EInnermostKind::Bool:		return {ELeafType::Bool, ELeafWidth::B8};
+	case EInnermostKind::S8:		return {ELeafType::IntS, ELeafWidth::B8};
+	case EInnermostKind::S16:		return {ELeafType::IntS, ELeafWidth::B16};
+	case EInnermostKind::S32:		return {ELeafType::IntS, ELeafWidth::B32};
+	case EInnermostKind::S64:		return {ELeafType::IntS, ELeafWidth::B64};
+	case EInnermostKind::U8:		return {ELeafType::IntU, ELeafWidth::B8};
+	case EInnermostKind::U16:		return {ELeafType::IntU, ELeafWidth::B16};
+	case EInnermostKind::U32:		return {ELeafType::IntU, ELeafWidth::B32};
+	case EInnermostKind::U64:		return {ELeafType::IntU, ELeafWidth::B64};
+	case EInnermostKind::F32:		return {ELeafType::Float, ELeafWidth::B32};
+	case EInnermostKind::F64:		return {ELeafType::Float, ELeafWidth::B64};
+	case EInnermostKind::UTF8:		return {ELeafType::Unicode, ELeafWidth::B8};
+	case EInnermostKind::UTF16:		return {ELeafType::Unicode, ELeafWidth::B16};
+	case EInnermostKind::UTF32:		return {ELeafType::Unicode, ELeafWidth::B32};
+	default:						check(false); return {ELeafType::Bool, ELeafWidth::B8};
+	}
+}
+
+static ELeafWidth ToLeafWidth(EInnermostKind Arithmetic)
+{
+	return ToLeafType(Arithmetic).Width;
+}
+
+static bool CompareLeaf(ELeafType Expected, EInnermostKind Actual)
+{
+	switch (Actual)
+	{
+		case EInnermostKind::Enum:		return Expected == ELeafType::Enum;
+		case EInnermostKind::Bool:		return Expected == ELeafType::Bool;
+		case EInnermostKind::S8:		return Expected == ELeafType::IntS;
+		case EInnermostKind::S16:		return Expected == ELeafType::IntS;
+		case EInnermostKind::S32:		return Expected == ELeafType::IntS;
+		case EInnermostKind::S64:		return Expected == ELeafType::IntS;
+		case EInnermostKind::U8:		return Expected == ELeafType::IntU;
+		case EInnermostKind::U16:		return Expected == ELeafType::IntU;
+		case EInnermostKind::U32:		return Expected == ELeafType::IntU;
+		case EInnermostKind::U64:		return Expected == ELeafType::IntU;
+		case EInnermostKind::F32:		return Expected == ELeafType::Float;
+		case EInnermostKind::F64:		return Expected == ELeafType::Float;
+		case EInnermostKind::UTF8:		return Expected == ELeafType::Unicode;
+		case EInnermostKind::UTF16:		return Expected == ELeafType::Unicode;
+		case EInnermostKind::UTF32:		return Expected == ELeafType::Unicode;	
+		default:						return false;
+	}	
+}
+
+static bool CompareInnermost(FMemberType Expected, EInnermostKind Actual)
+{
+	return Expected.IsStruct() ? Actual == EInnermostKind::Struct : CompareLeaf(Expected.AsLeaf().Type, Actual);
+}
+
+static void CheckSchema(FMemberSchema Expected, FLooseType Actual)
+{
+	check(Expected.InnerSchema == Actual.InnermostId);
+	check(Expected.NumInnerRanges == Actual.NumRanges);
+	check(CompareInnermost(Expected.GetInnermostType(), Actual.InnermostType.Kind));
+}
+
+FBuiltStruct* FBatchRun::FastenStruct(FLooseStruct Struct) const
+{
+	const FParameterIndexer& Names = History.LooseIds.Get(Struct.Id);
+	const FStructDeclaration& Declaration = History.Latest.Get(Struct.Id);
+
+	FMemberBuilder Out;
+	FastenMembers(Out, MakeArrayView(Struct.Members, Struct.Num), Names, Declaration);
+	return Out.BuildAndReset(Scratch, Declaration, Debug); 
+}
+
+void FBatchRun::FastenMembers(FMemberBuilder& Out, TConstArrayView<FLooseMember> Members, const FParameterIndexer& Names, const FStructDeclaration& Declaration) const
+{
+	// Fasten super
+	if (Names.Resolve(GetName(Members[0])) == NoId)
+	{
+		FStructMember Super = AsStruct(Members[0]);
+		Members.RightChopInline(1);
+		Out.AddSuperStruct(Super.Id, FastenStruct(Super));
+	}
+
+	// Fasten named members
+	FTargetLooseTypeIterator DeclIt(Declaration, History.Latest);
+	for (FLooseMember Member : Members)
+	{
+		FMemberSchema Declared = DeclIt.GrabSchema();
+		CheckSchema(Declared, Member.GetType());
+		if (Member)
+		{
+			// todo: decide pass down Declared.GetInnermostType() to get enum width or not
+			// maybe pass down innermost leaf size, 0 if not leaf
+			FBuiltValue Value = FastenMember(Member /*, Declared.GetInnermostType()*/);
+			FMemberId Name = Names.Resolve(GetName(Member)).Get();
+			Out.Add(Name, {Declared, Value});
+		}
+	}
+}
+
+FBuiltValue FBatchRun::FastenMember(FLooseMember Member) const
+{
+	FLooseType Type = Member.GetType();
+	if (Type.IsRange())
+	{
+		return { .Range = FastenRange(AsRange(Member)) };
+	}
+	else if (Type.IsStruct())
+	{
+		return { .Struct = FastenStruct(AsStruct(Member)) };
+	}
+	else if (Type.IsEnum())
+	{
+		return { .Leaf = FastenEnum(AsEnum(Member)) };
+	}
+	else
+	{
+		return { .Leaf = Member.GetValue().Data.Arithmetic };
+	}
+}
+
+FBuiltStruct* FBatchRun::FastenStruct(FStructMember Struct) const
+{
+	if (Struct.State == ELooseState::Source)
+	{
+		return RebuildStruct(AsSource(Struct, SourceBatch));
+	}
+	
+	FLooseStruct Loose = { .Id = Struct.Id, .Version = Struct.Version, .Num = IntCastChecked<uint16>(Struct.Num), .Members = Struct.Data.Loose };
+	if (Struct.State == ELooseState::Loose)
+	{
+		FTransformation Xform = Transformations[Struct.Id];
+		checkf(!Xform.Upgrades, TEXT("Struct should've been upgraded already"));
+		check(Xform.TargetSchema);
+		Loose = ReorderAsDeclaration(Loose, *Xform.TargetSchema, Scratch);
+	}
+		
+	return FastenStruct(Loose);
+}
+
+uint64 FBatchRun::FastenEnum(FEnumMember Enum) const
+{
+	if (Enum.State == ELooseState::Source)
+	{
+		return Enum.Data.Source;
+	}
+	//checkf(Enum.State == ELooseState::UpgradedLoose, TEXT("Enum not fully upgraded"));
+	FEnumFastener Fastener = Transformations[Enum.Id].Loose;
+	const FEnumeratorIndexer& Indexer = History.LooseIds.Get(Enum.Id); // opt: No need to lookup if no fastener...
+	return Enum.FlagMode ? Fasten(Fastener, Enum.Data.Loose.Enumerators.Flags, Indexer)
+						 : Fasten(Fastener, Enum.Data.Loose.Enumerators.Flat, Indexer);
+}
+
+FBuiltRange* FBatchRun::FastenRange(FRangeMember Range) const
+{
+	return FastenRange(Range.Type, Range.Meta, Range.Data);
+}
+
+FBuiltRange* FBatchRun::FastenRange(FLooseType Type, FLooseMetadata Meta, FRangeData Data) const
+{
+	if (Type.NumRanges > 1)
+	{
+		return FastenRanges(Type, Meta, MakeArrayView(Data.Ranges, Meta.Num));
+	}
+	else if (Type.IsStruct())
+	{
+		return FastenStructs(Type.InnermostId.Get().AsStructDeclId(), Meta, Data.Structs);
+	}
+	else if (Type.IsEnum())
+	{
+		return FastenEnums(Type.InnermostId.Get().AsEnum(), Type.InnermostType.Enum, Meta, Data);
+	}
+	else
+	{
+		return FastenArithmetics(Type.InnermostType.Kind, Meta.Num, Data.Leaves);
+	}
+}
+
+FBuiltRange* FBatchRun::FastenRanges(FLooseType Type, FLooseMetadata Meta, TConstArrayView64<FLooseNestedRange> Ranges) const
+{
+	if (Ranges.IsEmpty())
+	{
+		return nullptr;
+	}
+
+	FBuiltRange* Out = FBuiltRange::Create(Scratch, Ranges.Num(), sizeof(FBuiltRange*));
+	FBuiltRange** OutIt = reinterpret_cast<FBuiltRange**>(Out->Data);
+
+	FLooseType InnerType = Type;
+	InnerType.NumRanges--;
+
+	if (InnerType.IsRange())
+	{
+		for (FLooseNestedRange Range : Ranges)
+		{
+			*OutIt++ = FastenRanges(InnerType, Meta, MakeArrayView(Range.Data.Ranges, Range.Num));
+		}
+	}
+	else if (InnerType.IsStruct())
+	{
+		FDeclId Id = InnerType.InnermostId.Get().AsStructDeclId();
+		for (FLooseNestedRange Range : Ranges)
+		{
+			Meta.Num = Range.Num;
+			*OutIt++ = FastenStructs(Id, Meta, Range.Data.Structs);
+		}
+	}
+	else if (InnerType.IsEnum())
+	{
+		FEnumId Id = InnerType.InnermostId.Get().AsEnum();
+		FInnermostEnumType InnermostType = InnerType.InnermostType.Enum;
+		for (FLooseNestedRange Range : Ranges)
+		{
+			Meta.Num = Range.Num;
+			*OutIt++ = FastenEnums(Id, InnermostType, Meta, Range.Data);
+		}
+	}
+	else
+	{
+		for (FLooseNestedRange Range : Ranges)
+		{
+			*OutIt++ = FastenArithmetics(InnerType.InnermostType.Kind, Range.Num, Range.Data.Leaves);
+		}
+	}
+
+	return Out;
+}
+
+FBuiltRange* FBatchRun::FastenStructs(FDeclId Id, FLooseMetadata Meta, FStructRange Structs) const
+{
+	if (Meta.Num == 0)
+	{
+		return nullptr;
+	}
+	else if (Meta.InnermostSchema)
+	{
+		FStructSchemaHandle Schema = { FStructSchemaId(Meta.InnermostSchema.Get()), SourceBatch };
+		return RebuildStructs(FStructRangeView(Meta.Num, ViewSkippableSlice(Structs.Source), Schema), Id);
+	}
+
+	const FParameterIndexer& Names = History.LooseIds.Get(Id);
+	const FStructDeclaration& Declaration = History.Latest.Get(Id);
+
+	FBuiltRange* Out = FBuiltRange::Create(Scratch, Meta.Num, sizeof(FBuiltStruct*));
+	FBuiltStruct** OutIt = reinterpret_cast<FBuiltStruct**>(Out->Data);
+
+	FMemberBuilder Tmp;
+	for (FLooseValue Struct : TConstArrayView64<FLooseValue>(Structs.Loose, Meta.Num))
+	{
+		FastenMembers(Tmp, MakeArrayView(Struct.Data.Struct.Loose, Struct.Meta.Num), Names, Declaration);
+		*OutIt++ = Tmp.BuildAndReset(Scratch, Declaration, Debug); 
+	}
+
+	return Out;
+}
+
+FBuiltRange* FBatchRun::FastenEnums(FEnumId Id, FInnermostEnumType Type, FLooseMetadata Meta, FRangeData Data) const
+{
+	if (Meta.Num == 0)
+	{
+		return nullptr;
+	}
+	else if (Meta.InnermostSchema)
+	{
+		SIZE_T EnumSize = SizeOf(History.Latest.Get(Id).Width.Get());
+		return FBuiltRange::CreateIf(Scratch, Meta.Num, EnumSize, Data.Leaves); 
+	}
+
+	// todo: enum range fastening
+	unimplemented();
+
+	return nullptr;
+}
+
+FBuiltRange* FBatchRun::FastenArithmetics(EInnermostKind Kind, uint64 Num, const void* Data) const
+{
+	if (Num == 0)
+	{
+		return nullptr;
+	}
+
+	SIZE_T LeafSize = SizeOf(ToLeafWidth(Kind));
+	FBuiltRange* Out = FBuiltRange::Create(Scratch, Num, LeafSize);
+	if (Kind == EInnermostKind::Bool)
+	{
+		FBoolRangeView(Data, Num).Copy(Out->Data, Num); // bits -> bools
+	}
+	else
+	{
+		FMemory::Memcpy(Out->Data, Data, Num * LeafSize);
+	}
+
+	return Out;
+}
+
+FBuiltStruct* FBatchRun::RebuildStruct(FStructView Source) const
+{
+	return RebuildStruct(Source, RuntimeIds[Source.Schema.Id]);
+}
+
+FBuiltStruct* FBatchRun::RebuildStruct(FStructView Source, FDeclId Id) const
+{	
+	FMemberBuilder Out;
+	RebuildMembers(Out, Source, Id);
+	return Out.BuildAndReset(Scratch, History.Latest.Get(Id), Debug);
+}
+
+void FBatchRun::RebuildMembers(FMemberBuilder& Out, FStructView Source, FDeclId Id) const
+{
+	FMemberReader It(Source);
+	if (It.HasMore() && IsSuper(It.PeekType()))
+	{
+		FStructView Super = It.GrabStruct();
+		FDeclId SuperId = RuntimeIds[Super.Schema.Id];
+		Out.AddSuperStruct(SuperId, RebuildStruct(Super, SuperId));
+	}
+
+	while (It.HasMore())
+	{
+		FMemberId Name = It.PeekNameUnchecked();
+		Out.Add(Name, RebuildValue(It));
+	}
+}
+
+FTypedValue	FBatchRun::RebuildValue(FMemberReader& It) const
+{
+	FMemberType Type = It.PeekType();
+	switch (Type.GetKind())
+	{
+	case EMemberKind::Leaf:		return RebuildValue(It.GrabLeaf());
+	case EMemberKind::Range:	return RebuildValue(Type, It.GrabRange());
+	case EMemberKind::Struct:	return RebuildValue(It.GrabStruct());
+	}
+	check(false);
+	return {};
+}
+
+FTypedValue	FBatchRun::RebuildValue(FLeafView In) const
+{
+	FOptionalEnumId Enum = In.Enum ? ToOptional(RuntimeIds[In.Enum.Get()]) : NoId;
+	FBuiltMember Out(FMemberId{}, In.Leaf, Enum, GetIntegerValue(In));
+	return {Out.Schema, Out.Value};
+}
+
+static uint16 CountRanges(FRangeSchema Schema)
+{
+	if (!Schema.ItemType.IsRange())
+	{
+		return 1;
+	}
+
+	uint32 Idx = 0;
+	while (Schema.NestedItemTypes[Idx++].IsRange());
+	return IntCastChecked<uint16>(1 + Idx);
+}
+
+FTypedValue	FBatchRun::RebuildValue(FMemberType Type, FRangeView Range) const
+{
+	FRangeSchema InSchema = GetSchema(Range);
+	FMemberSchema OutSchema = { Type, InSchema.ItemType, CountRanges(InSchema), RuntimeIds[InSchema.InnermostSchema], 
+								InSchema.ItemType.IsRange() ? InSchema.NestedItemTypes : nullptr };
+	FBuiltRange* OutValue = RebuildRange(Range);
+	return { OutSchema, { .Range = OutValue } };
+}
+
+FTypedValue FBatchRun::RebuildValue(FStructView Struct) const
+{
+	FDeclId Id = RuntimeIds[Struct.Schema.Id];
+	FMemberSchema Schema = { DefaultStructType, DefaultStructType, 0,  FInnerId(Id), nullptr };
+	return { Schema, { .Struct = RebuildStruct(Struct, Id) } };
+}
+
+FBuiltRange* FBatchRun::RebuildRange(FRangeView Range) const
+{
+	if (Range.Num())
+	{
+		switch (Range.GetItemType().GetKind())
+		{
+		case EMemberKind::Leaf:		return RebuildLeaves(Range.AsLeaves());
+		case EMemberKind::Range:	return RebuildRanges(Range.AsRanges());
+		case EMemberKind::Struct:	return RebuildStructs(Range.AsStructs());
+		}
+	}
+	return nullptr;
+}
+
+
+FBuiltRange* FBatchRun::RebuildLeaves(FLeafRangeView Leaves) const
+{
+	FUnpackedLeafType Leaf = GetType(Leaves);
+	FBuiltRange* Out = FBuiltRange::Create(Scratch, Leaves.Num(), SizeOf(Leaf.Width));
+	if (Leaf.Type == ELeafType::Bool)
+	{
+		Leaves.AsBools().Copy(Out->Data, Leaves.Num()); // bits -> bools
+	}
+	else
+	{
+		FMemory::Memcpy(Out->Data, GetValues(Leaves), Leaves.Num() * SizeOf(Leaf.Width));
+	}
+
+	return Out;
+}
+
+FBuiltRange* FBatchRun::RebuildRanges(FNestedRangeView Ranges) const
+{
+	FBuiltRange* Out = FBuiltRange::Create(Scratch, Ranges.Num(), sizeof(FBuiltRange*));
+	FBuiltRange** OutIt = reinterpret_cast<FBuiltRange**>(Out->Data);
+	for (FRangeView Range : Ranges)
+	{
+		*OutIt++ = RebuildRange(Range);
+	}
+	return Out;
+}
+
+FBuiltRange* FBatchRun::RebuildStructs(FStructRangeView Structs) const
+{
+	return RebuildStructs(Structs, RuntimeIds[Structs.GetSchema().Id]);
+}
+
+FBuiltRange* FBatchRun::RebuildStructs(FStructRangeView Structs, FDeclId Id) const
+{
+	const FStructDeclaration& Declaration = History.Latest.Get(Id);
+
+	FBuiltRange* Out = FBuiltRange::Create(Scratch, Structs.Num(), sizeof(FBuiltStruct*));
+	FBuiltStruct** OutIt = reinterpret_cast<FBuiltStruct**>(Out->Data);
+	FMemberBuilder Tmp;
+	for (FStructView Struct : Structs)
+	{
+		RebuildMembers(Tmp, Struct, Id);
+		*OutIt++ = Tmp.BuildAndReset(Scratch, Declaration, Debug);
+	}
+	return Out;
+}
+
+////////////////////////////////////////////////////////////////////////////////////////////////
+
+FBatchUpgrader::FBatchUpgrader(const FSchemaBatch& InBatch, TConstArrayView<FInnerId> InRuntimeIds, FIdIndexerBase& Ids)
+: Batch(InBatch)
+, RuntimeIds(InRuntimeIds)
+{}
+
+static bool InnerSchemasMatches(const FStructSchema& Schema, const TBitArray<>& MatchingSchemas)
+{
+	for (FSchemaId Inner : MakeArrayView(Schema.GetInnerSchemas(), Schema.NumInnerSchemas))
+	{
+		if (!MatchingSchemas[Inner.Idx])
+		{
+			return false;
+		}
+	}
+
+	return true;
+}
+
+static bool UnmatchOuterSchemas(TBitArray<>& MatchingSchemas, TSchemaRange<const FStructSchema> Schemas)
+{
+	bool bUnmatched = false;
+	for (int32 Idx = 0; const FStructSchema& Schema : Schemas)
+	{
+		if (MatchingSchemas[Idx] && !InnerSchemasMatches(Schema, MatchingSchemas))
+		{
+			bUnmatched = true;
+			MatchingSchemas[Idx] = false;
+		}
+		++Idx;
+	}
+	return bUnmatched;
+}
+
+bool FBatchUpgrader::MatchSchemas(const IDeclarations& Types)
+{
+	check(MatchingSchemas.IsEmpty());
+
+	const int32 Num = Batch.NumSchemas;
+	const int32 NumStructs = Batch.NumStructSchemas;
+	const int32 NumEnums = Num - NumStructs;
+	MatchingSchemas.SetNumUninitialized(Num);
+
+	// Cache all struct declarations
+	Declarations.Reserve(NumStructs);
+	for (FDeclId Id : FInnerIds({RuntimeIds}).GetDeclIds(NumStructs))
+	{
+		Declarations.Emplace(Types.Find(Id));
+	}
+
+	// Match struct schemas against declarations
+	int32 Idx = 0;
+	TSchemaRange<const FStructSchema> StructSchemas = GetStructSchemas(Batch);
+	for (const FStructSchema& Schema : StructSchemas)
+	{
+		const FStructDeclaration* Decl = Declarations[Idx];
+		MatchingSchemas[Idx] = Decl && MatchStruct(*Decl, Schema, {RuntimeIds});
+		++Idx;
+	}
+
+	// Match enum schemas against declarations
+	for (const FEnumSchema& Schema : GetEnumSchemas(Batch))
+	{
+		const FEnumDeclaration* Decl = Types.Find(RuntimeIds[Idx].AsEnum());
+		MatchingSchemas[Idx] = Decl && MatchEnum(*Decl, Schema);
+		++Idx;
+	}
+
+	// Flag outer structs as non-matching when inner structs or enums dont match
+	// using a somewhat naive recursive approach that can be optimized	
+	if (MatchingSchemas.Contains(false))
+	{
+		while (UnmatchOuterSchemas(/* in-out */ MatchingSchemas, StructSchemas));
+
+		return false;
+	}
+
+	return true;
+}
+
+TArray<FUpgradedInstance> FBatchUpgrader::Upgrade(TConstArrayView<FStructView> Old, FHistory& Upgrades, FScratchAllocator& Scratch)
+{
+	check(!MatchingSchemas.IsEmpty());
+	
+	Upgrades.Compile();
+	
+	FBatchRun Run(Upgrades, Scratch, RuntimeIds);
+	Run.CreateTransformations(MatchingSchemas, Batch);
+
+	TArray<FUpgradedInstance> Out;
+	Out.Reserve(Old.Num());
+	for (FStructView Source : Old)
+	{
+		Out.Emplace(Run.UpgradeOrRebuild(Source));
+	}
+
+	return Out;
+}
+
+////////////////////////////////////////////////////////////////////////////////////////////////
+
+static FType IndexType(FLiteralIndexerBase& Ids, FTypename In)
+{
+	return {Ids.IndexScopeLiteral(In.Scope), FTypenameId(Ids.IndexTypenameLiteral(In.Name))};
+}
+
+static FStructId IndexStruct(FLiteralIndexerBase& Ids, FTypename In)
+{
+	return Ids.IndexStruct(IndexType(Ids, In));
+}
+
+////////////////////////////////////////////////////////////////////////////////////////////////
+
+struct FMidOp
+{
+	Op	 					Function;
+	uint8					NumOptionalInputs;	// Excluding inner parameters
+	uint8					NumRequiredInputs;	// Excluding inner parameters
+	uint8					NumOutputs;			// Excluding inner parameters
+};
+
+struct FMidParam
+{
+	FParameterId			Name;
+	FLooseType				Type;
+	uint8					InnerSlice;
+};
+
+struct FInnerSlice
+{
+	int32 					ItemIdx;
+	int32					NumItems;
+
+	explicit operator bool() { return !!NumItems;}
+};
+
+template<typename T>
+struct TInnerSlices
+{
+	uint8 AddSlice(TConstArrayView<T> Slice)
+	{
+		check(Slice.Num() > 0);
+		Slices.Emplace(Items.Num(), Slice.Num());
+		Items.Append(Slice);
+		return IntCastChecked<uint8>(Slices.Num());
+	}
+
+	TConstArrayView<T> GetSlice(uint8 Handle) const
+	{
+		check(Handle > 0);
+		FInnerSlice Slice = Slices[Handle - 1];
+		check(Slice.ItemIdx + Slice.NumItems <= Items.Num());
+		return MakeArrayView(Items.GetData() + Slice.ItemIdx, Slice.NumItems);
+	}
+
+	TArray<FInnerSlice, TInlineAllocator<16>>	Slices;			// todo switch to scratch or vanilla allocs
+	TArray<T, TInlineAllocator<32>>				Items;			// todo switch to scratch or vanilla allocs
+};
+
+struct FOps final : IOps
+{
+	FOps(FDeclId Id, uint16 Ver, FLiteralIndexerBase& InIds, FLooseIndexers& Indexers)
+	: Owner(Id)
+	, Version(Ver)
+	, Ids(InIds)
+	, OwnerIndexer(Indexers.Get(Id))
+	, InnerIndexers(Indexers)
+	{}
+
+	const FDeclId								Owner;
+	const uint16								Version;
+	FOptionalVersion							NextVersion;
+	FLiteralIndexerBase&						Ids;
+	FParameterIndexer&							OwnerIndexer;
+	FLooseIndexers&								InnerIndexers;
+	TArray<FMidOp, TInlineAllocator<8>>			Ops;			// todo switch to scratch or vanilla allocs
+	TArray<FMidParam, TInlineAllocator<32>>		Parameters;		// todo switch to scratch or vanilla allocs
+	TInnerSlices<FMidParam>						InnerMembers;
+	TInnerSlices<FEnumeratorId>					InnerEnumerators;
+	FOptionalDeclId								Redirect;
+	int32										NumInputs = 0;
+	int32										NumOutputsFlat = 0;
+	int32										NumDecompositions = 0;
+	int32										NumStructDecompositions = 0;
+
+	int32										NumMatches() const				{ return NumInputs + NumDecompositions; }
+	int32										NumRecompositions() const		{ return InnerEnumerators.Slices.Num(); }
+	int32										NumRecompositionIds() const		{ return InnerEnumerators.Items.Num(); }
+
+	uint8 IndexEnumerators(FEnumId Id, std::initializer_list<const char*> Enumerators)
+	{
+		FEnumeratorIndexer& InnerIndexer = InnerIndexers.Get(Id);
+
+		TArray<FEnumeratorId, TInlineAllocator<64>> Slice;
+		for (FAnsiStringView Enumerator : Enumerators)
+		{
+			Slice.Add(InnerIndexer.Index(Ids.IndexEnumeratorLiteral(Enumerator)));
+		}
+
+		return InnerEnumerators.AddSlice(Slice);
+	}
+
+	uint8 IndexMembers(FStructId Id, TConstArrayView<FParameter> Members)
+	{
+		// opt: reuse slices, special case iota slices
+
+		FParameterIndexer& InnerIndexer = InnerIndexers.Get(FDeclId(Id));
+		TArray<FMidParam, TInlineAllocator<64>> Slice;
+		Slice.Reserve(Members.Num());
+		for (FParameter Member : Members)
+		{
+			Slice.Add(IndexParameter(Member, InnerIndexer));
+		}
+		return InnerMembers.AddSlice(Slice);
+	}
+
+	virtual FEnumId	NameEnum(FTypename Enum) override
+	{
+		return Ids.IndexEnum(IndexType(Ids, Enum));
+	}
+
+	virtual FStructId NameStruct(FTypename Struct) override
+	{
+		return IndexStruct(Ids, Struct);
+	}
+	
+	virtual FEnumeratorNames NameEnumerators(FEnumId Id, std::initializer_list<const char*> Enumerators) override
+	{
+		return { Id, IndexEnumerators(Id, Enumerators) };
+	}
+
+	virtual FEnumFlagNames NameFlags(FEnumId Id, std::initializer_list<const char*> Flags) override
+	{
+		return { Id, IndexEnumerators(Id, Flags) };
+	}
+
+	virtual FStructMemberNames NameMembers(FStructId Id, TConstArrayView<FParameter> Members) override
+	{
+		return { Id, IndexMembers(Id, Members) };
+	}
+
+	bool IsOpen() const
+	{
+		return !NextVersion;
+	}
+	
+	virtual void Bump(uint16 Next) override
+	{
+		check(IsOpen());
+		NextVersion = Next;
+	}
+
+	virtual void Rename(FTypename To, uint16 ToVersion) override
+	{
+		check(IsOpen());
+		Redirect = FDeclId(IndexStruct(Ids, To));
+		NextVersion = ToVersion;
+	}
+
+	uint32 CountFlat(TConstArrayView<FMidParam> Params) const
+	{
+		uint32 Out = Params.Num();
+		for (FMidParam Param : Params)
+		{
+			if (uint8 Slice = Param.InnerSlice)
+			{
+				if (Param.Type.IsStruct())
+				{
+					Out += CountFlat(InnerMembers.GetSlice(Slice));
+				}
+			}
+		}
+		return Out;
+	}
+
+	virtual void Add(TConstArrayView<FParameter> Ins, TConstArrayView<FParameter> Outs, Op Func) override
+	{
+		check(IsOpen());
+
+		uint8 NumOptionalInputs = IntCastChecked<uint8>(Algo::CountIf(Ins, [](FParameter P){ return P.bOptional; }));
+		uint8 NumRequiredInputs = IntCastChecked<uint8>(Ins.Num() - NumOptionalInputs);
+		Ops.Emplace(Func, NumOptionalInputs, NumRequiredInputs, IntCastChecked<uint8>(Outs.Num()));
+		IndexParameters(Ins);
+		CountDecompositions(MakeArrayView(Parameters).Right(Ins.Num()));
+		NumInputs += Ins.Num();
+		IndexParameters(Outs);
+		NumOutputsFlat += CountFlat(MakeArrayView(Parameters).Right(Outs.Num()));
+	}
+
+	void IndexParameters(TConstArrayView<FParameter> Params)
+	{
+		Parameters.AddUninitialized(Params.Num());
+		FMidParam* OutIt = Parameters.GetData() + Parameters.Num() - Params.Num();
+		for (FParameter Param : Params)
+		{
+			*OutIt++ = IndexParameter(Param, OwnerIndexer);
+		}
+	}
+
+	FMidParam IndexParameter(FParameter In, FParameterIndexer& StructIndexer) const
+	{
+		check(In.NumRanges <= MaxRangeNesting);
+		check(In.InnermostHandle.Id || In.InnermostHandle.Version == 0);
+
+		FMemberId Member = Ids.IndexMemberLiteral(FAnsiStringView(In.Name, In.NameLen));
+		FParameterId Name = StructIndexer.Index(Member);
+		FLooseType Type = { In.bOptional, In.NumRanges, In.InnermostType, In.InnermostHandle.Version, In.InnermostHandle.Id };
+
+		return {Name, Type, In.InnermostHandle.Slice};
+	}
+
+	void CountDecompositions(TConstArrayView<FMidParam> Inputs)
+	{
+		for (FMidParam Input : Inputs)
+		{
+			if (Input.InnerSlice)
+			{
+				++NumDecompositions;
+				NumStructDecompositions += Input.Type.IsStruct();
+			}
+		}
+	}
+};
+
+////////////////////////////////////////////////////////////////////////////////////////////////
+
+const FLooseSchema* FSchemaCache::Obtain(FDeclId Id, const IDeclarations& Types, FParameterIndexer& ParamIds)
+{
+	if (const FLooseSchemaPtr* Out = Cache.Find(Id))
+	{
+		return Out->Get();
+	}
+	if (const FStructDeclaration* Struct = Types.Find(Id))
+	{
+		return Cache.Emplace(Id, CreateTargetSchema(*Struct, Types, ParamIds)).Get();
+	}
+
+	// Don't cache missing declarations for now
+	return nullptr;
+}
+
+////////////////////////////////////////////////////////////////////////////////////////////////
+
+static FEnumeratorId GetInnerName(FEnumeratorId Inner) { return Inner; }
+static FParameterId GetInnerName(FMidParam Inner) { return Inner.Name; }
+
+template <typename InnerType>
+static bool IsDeclarationOrder(TConstArrayView<InnerType> Order)
+{
+	uint32 N = 0;
+	for (InnerType Inner : Order)
+	{
+		if (GetInnerName(Inner).Idx != N++)
+		{
+			return false;
+		}
+
+	}
+	return true;
+}
+
+static TSet<FParameterId> CreateInputOrder(TConstArrayView<FMidParam> Order)
+{
+	if (IsDeclarationOrder(Order))
+	{
+		return {};
+	}
+		
+	TSet<FParameterId> Out;
+	Out.Reserve(Order.Num());
+	for (FMidParam Param : Order)
+	{
+		Out.Add(Param.Name);
+	}
+	return Out;
+}
+
+static TSet<FEnumeratorId> CreateInputOrder(TConstArrayView<FEnumeratorId> Order)
+{
+	return IsDeclarationOrder(Order) ? TSet<FEnumeratorId>() : TSet<FEnumeratorId>(Order);
+}
+
+////////////////////////////////////////////////////////////////////////////////////////////////
+
+static FLooseSchemaPtr CreateInputSchema(uint16 Version, TConstArrayView<FMidParam> Order)
+{
+	const uint32 Size = offsetof(FLooseSchema, Members) + Order.Num() * sizeof(FLooseParameter);
+	FLooseSchema Header = {Version, IntCastChecked<uint16>(Order.Num())};
+	FLooseSchema* Out = new (FMemory::MallocZeroed(Size, alignof(FLooseSchema))) FLooseSchema(Header);
+
+	FLooseParameter* It = Out->Members;
+	for (FMidParam Param : Order)
+	{
+		new (It++) FLooseParameter{Param.Name, Param.Type};
+	}
+
+	return FLooseSchemaPtr(Out);
+}
+
+////////////////////////////////////////////////////////////////////////////////////////////////
+
+static uint64 CopyFlat(FOutParameter* Dst, TConstArrayView<FMidParam> Params, const FOps& Ops)
+{
+	FOutParameter* It = Dst;
+	for (FMidParam Param : Params)
+	{
+		FOutParameter& Out = *It++;
+		Out = {Param.Name, 0, 0, Param.Type};
+
+		if (uint8 Slice = Param.InnerSlice)
+		{
+			if (Param.Type.InnermostType.Kind == EInnermostKind::Struct)
+			{
+				TConstArrayView<FMidParam> Inners = Ops.InnerMembers.GetSlice(Slice);
+				Out.NumInners = static_cast<uint8>(Inners.Num());
+				It += CopyFlat(It, Inners, Ops);
+			}
+			else
+			{
+				check(Param.Type.InnermostType.Kind == EInnermostKind::Enum);
+				TConstArrayView<FEnumeratorId> Inners = Ops.InnerEnumerators.GetSlice(Slice);
+				Out.NumInners = static_cast<uint8>(Inners.Num());
+				Out.RecompIdx = Slice - 1;
+			}
+		}
+	}
+
+	return static_cast<uint64>(It - Dst);
+}
+
+static void CompileDecomposition(uint8*& InOutIt, FMidParam Input, TInnerSlices<FMidParam> Members, TInnerSlices<FEnumeratorId> Enumerators)
+{
+	FInnerId Id = Input.Type.InnermostId.Get();
+	uint16 Version = Input.Type.InnermostVersion;
+	if (Input.Type.IsStruct())
+	{
+		TConstArrayView<FMidParam> Inners = Members.GetSlice(Input.InnerSlice);
+		new (InOutIt) FStructDecomposition{ Id.AsStructDeclId(), Version,  CreateInputOrder(Inners), /* set in Link */ nullptr, 
+											CreateInputSchema(Input.Type.InnermostVersion, Inners) };
+		InOutIt += sizeof(FStructDecomposition);
+	}
+	else
+	{
+		check(Input.Type.IsEnum());
+		TConstArrayView<FEnumeratorId> Inners = Enumerators.GetSlice(Input.InnerSlice);
+		new (InOutIt) FEnumDecomposition{ CreateInputOrder(Inners) };
+		InOutIt += sizeof(FEnumDecomposition);
+	}
+}
+
+static int32 CompileRecomposition(uint16& OutJump, FEnumeratorId* Dst, TConstArrayView<FEnumeratorId> Src)
+{				
+	// The same FNamedEnumerators/FNamedFlags slice can be reused for multiple outputs -- copy enumerators once
+	if (OutJump)
+	{
+		return 0;
+	}
+
+	static_assert(sizeof(OutJump) == sizeof(*Dst), "Jump offset is relative to sel and assumes jump offset and FEnumeratorId have same size");
+	OutJump = IntCastChecked<uint16>(reinterpret_cast<uint16*>(Dst) - &OutJump);
+	FMemory::Memcpy(Dst, Src.GetData(), Src.NumBytes());
+
+	return Src.Num();
+}
+
+static void CompileOps(FVersionUpgrades& Out, const FOps& In)
+{
+	using MemberMatches = TArray<FMemberMatch, TInlineAllocator<1>>;
+	TArray<MemberMatches> AllMatches;
+	AllMatches.SetNum(In.OwnerIndexer.Num());
+
+	TConstArrayView<FMidParam> Params = In.Parameters;
+	FOutParameter* Outputs = const_cast<FOutParameter*>(Out.GetOutputs());
+	FOp* OutOps = const_cast<FOp*>(Out.GetOps());
+	uint32* DecompOffsets = const_cast<uint32*>(Out.GetDecompOffsets());
+	uint16* RecompJumps = const_cast<uint16*>(Out.GetRecompJumps());
+	FEnumeratorId* RecompIdIt = reinterpret_cast<FEnumeratorId*>(RecompJumps + Out.NumRecompositions);
+	uint8* DecompDataIt = Align(reinterpret_cast<uint8*>(Outputs + In.NumOutputsFlat), alignof(FStructDecomposition));
+	uint16 DecompIdx = 0;
+	uint32 OutputIdx = 0;
+	TConstArrayView<FMidOp> InOps = In.Ops;
+	for (int32 OpIdx = 0; OpIdx < InOps.Num(); ++OpIdx)
+	{
+		const FMidOp Op = InOps[OpIdx];
+
+		// Compile inputs into member matches
+		check(Op.NumOptionalInputs + Op.NumRequiredInputs <= 0xFF);
+		uint8 InputIdx = 0;
+		for (FMidParam Input : PopFront(Params, Op.NumOptionalInputs + Op.NumRequiredInputs))
+		{
+			FOptionalMemberId Foo = In.OwnerIndexer.Resolve(Input.Name);
+
+			// Generate input decomposition match
+			if (Input.InnerSlice)
+			{
+				FMatchIdx Idx = { .Decomp = DecompIdx };	
+				AllMatches[Input.Name.Idx].Add({ .Type = Input.Type, .Idx = Idx, .InputIdx = InputIdx, .Recipient = FMemberMatch::DecomposeRecipient });
+				
+				// Write offset and generate decomposition
+				DecompOffsets[DecompIdx] = IntCastChecked<uint32>(DecompDataIt - reinterpret_cast<uint8*>(&Out));
+				CompileDecomposition(/* in-out */ DecompDataIt, Input, In.InnerMembers, In.InnerEnumerators);
+				++DecompIdx;
+			}
+
+			// Copy normal input match
+			FMatchIdx Idx = { .Op = static_cast<uint16>(OpIdx) };
+			AllMatches[Input.Name.Idx].Add({ .Type = Input.Type, .Idx = Idx, .InputIdx = InputIdx, .Recipient = 0 });
+
+			++InputIdx;
+		}
+
+		// Copy out parameters
+		TConstArrayView<FMidParam> OutParams = PopFront(/* in-out */ Params, Op.NumOutputs);
+		uint64 NumOutputsFlat = CopyFlat(Outputs + OutputIdx, OutParams, In);
+		check(NumOutputsFlat == In.CountFlat(OutParams));
+
+		// Generate enum recompositions
+		for (FMidParam Output : OutParams)
+		{
+			if (Output.InnerSlice && Output.Type.InnermostType.Kind == EInnermostKind::Enum)
+			{
+				uint8 RecompIdx = Output.InnerSlice - 1;
+				RecompIdIt += CompileRecomposition(/* out */ RecompJumps[RecompIdx], RecompIdIt, In.InnerEnumerators.GetSlice(Output.InnerSlice));
+			}
+		}
+
+		// Copy FMidOp -> FOp
+		OutOps[OpIdx] = {	.Function = Op.Function,
+							.Outputs = { OutputIdx, IntCastChecked<uint8>(NumOutputsFlat) },
+							.NumOptionalInputs = Op.NumOptionalInputs,
+							.NumRequiredInputs = Op.NumRequiredInputs };
+
+		OutputIdx += NumOutputsFlat;
+	}
+	check(OutputIdx == In.NumOutputsFlat);
+	check(RecompIdIt == reinterpret_cast<FEnumeratorId*>(RecompJumps + Out.NumRecompositions) + In.InnerEnumerators.Items.Num());
+	check(Out.GetOutputs() == AlignPtr<FOutParameter>(RecompIdIt));
+
+	// Copy matches
+	FMiniSlice* NameMatchIt = Out.NameMatches;
+	FMemberMatch* Matches = const_cast<FMemberMatch*>(Out.GetMatches());
+	uint32 MatchIdx = 0;
+	for (const MemberMatches& NameMatches : AllMatches)
+	{
+		FMemory::Memcpy(Matches + MatchIdx, NameMatches.GetData(), NameMatches.NumBytes());
+		new (NameMatchIt++) FMiniSlice(MatchIdx, NameMatches.Num());
+		MatchIdx += NameMatches.Num(); 
+	}
+	check(MatchIdx == Out.NumMatches);
+}
+
+static void CompileVersion(FVersionUpgrades& Out, const FOps& In)
+{
+	Out.NumNames = In.OwnerIndexer.Num();
+	Out.NumMatches = static_cast<uint32>(In.NumMatches());
+	Out.NumOps = IntCastChecked<uint16>(In.Ops.Num());
+	Out.NumDecompositions = IntCastChecked<uint16>(In.NumDecompositions);
+	Out.NumRecompositions = IntCastChecked<uint16>(In.NumRecompositions());
+	const FEnumeratorId* RecompIds = AlignPtr<FEnumeratorId>(Out.GetRecompJumps() + In.NumRecompositions());
+	const FOutParameter* Outputs = AlignPtr<FOutParameter>(RecompIds + In.NumRecompositionIds());
+	Out.OutputsOffset = IntCastChecked<uint32>((const uint8*)Outputs - (uint8*)&Out);
+	Out.NextVersion = In.NextVersion;
+
+	CompileOps(Out, In);
+}
+
+static uint32 CalculateVersionUpgradesSize(const FOps& Ops)
+{
+	uint32 Out = offsetof(FVersionUpgrades, NameMatches) + Ops.OwnerIndexer.Num() * sizeof(FMiniSlice);
+	Out = Align(Out + Ops.NumMatches() * sizeof(FMemberMatch), alignof(FMemberMatch));
+	Out = Align(Out + Ops.Ops.Num() * sizeof(FOp), alignof(FOp));
+	Out = Align(Out + Ops.NumDecompositions * sizeof(uint32), alignof(uint32));
+	Out = Align(Out + Ops.NumRecompositions() * sizeof(uint16), alignof(uint16));
+	Out = Align(Out + Ops.NumRecompositionIds() * sizeof(FEnumeratorId), alignof(FEnumeratorId));
+	Out = Align(Out + Ops.NumOutputsFlat * sizeof(FOutParameter), alignof(FOutParameter));
+	static_assert(alignof(FStructDecomposition) == alignof(FEnumDecomposition));
+	Out = Align(Out + Ops.NumStructDecompositions * sizeof(FStructDecomposition), alignof(FStructDecomposition));
+	Out = Align(Out + (Ops.NumDecompositions - Ops.NumStructDecompositions) * sizeof(FEnumDecomposition), alignof(FEnumDecomposition));
+
+	return Out;
+}
+
+void FHistory::CompileStruct(TArrayView<FOpsPtr> AllOps)
+{
+	const FDeclId Id = FDeclId(AllOps[0]->Owner);
+	const uint32 StructSize = sizeof(FStructUpgrades) + AllOps.Num() * sizeof(FVersionUpgrades*);
+
+	uint32 VersionOffset = StructSize;
+	TArray<uint32, TInlineAllocator<32>> VersionOffsets;
+	for (const FOpsPtr& Ops : AllOps)
+	{
+		VersionOffset = Align(VersionOffset, alignof(FVersionUpgrades));
+		VersionOffsets.Add(VersionOffset);
+		VersionOffset += CalculateVersionUpgradesSize(*Ops);
+	}
+
+	FOptionalDeclId Redirect = AllOps.Last()->Redirect;
+	FOptionalVersion RedirectVersion = AllOps.Last()->NextVersion;
+	FStructUpgrades Header = { .Id = Id, .FinalId = Redirect.GetOr(Id), .FirstVersion = AllOps[0]->Version, .NumVersions = IntCastChecked<uint16>(AllOps.Num()) };
+	Header.FinalVersion = RedirectVersion ? RedirectVersion.Get() : Header.FirstVersion + Header.NumVersions - 1;
+	// Zeroed memory expected by CompileRecomposition
+	FStructUpgrades* Out = new (FMemory::MallocZeroed(VersionOffset, alignof(FStructUpgrades))) FStructUpgrades(Header);
+
+	for (int32 Idx = 0, Num = AllOps.Num(); Idx < Num; ++Idx)
+	{
+		FVersionUpgrades* Version = reinterpret_cast<FVersionUpgrades*>(reinterpret_cast<uint8*>(Out) + VersionOffsets[Idx]);
+		Out->Versions[Idx] = Version;
+		CompileVersion(*Version, *AllOps[Idx]);
+	}
+	
+	CompiledUpgrades.Emplace(Id, Out);
+}
+
+static TArrayView<FOpsPtr> GrabOwnerRange(FOpsPtr*& It, FOpsPtr* End)
+{
+	check(*It != nullptr);
+	FOpsPtr* First = It++;
+	for (FStructId Owner = (**First).Owner; It != End && *It && (**It).Owner == Owner; ++It);
+	int32 Num = It - First;
+	for (; It != End && *It == nullptr; ++It);
+	return MakeArrayView(First, Num);
+}
+
+void FHistory::Compile()
+{
+	if (CommitQueue.IsEmpty())
+	{
+		return;
+	}
+
+	FOpsPtr* It = CommitQueue.FindByPredicate([=](FOpsPtr& Ops){ return Ops != nullptr; });
+	for (FOpsPtr* End = GetEnd(CommitQueue); It != End;)
+	{
+		CompileStruct(GrabOwnerRange(/* in-out */ It, End));
+	}
+
+	Link(CommitQueue);
+
+	CommitQueue.Empty();
+	CommitScratch.Reset();
+}
+
+void FHistory::Erase(FStructId Owner)
+{
+	// todo: track inverse link dependencies and drop/disable/whatever compiled upgrades linked against owner
+
+	LatestSchemas.Cache.Remove(FDeclId(Owner));
+
+	if (CompiledUpgrades.Remove(FDeclId(Owner)) ||
+		CommitQueue.IsEmpty())
+	{
+		return;
+	}
+
+	// Don't reorder, CompileStruct requires Ops to be sorted by version
+	if (FOpsPtr* It = CommitQueue.FindByPredicate([=](FOpsPtr& Ops){ return Ops && Ops->Owner == Owner; }))
+	{
+		for (FOpsPtr* End = GetEnd(CommitQueue); It != End && *It && (**It).Owner == Owner; ++It)
+		{
+			It->Reset();
+		}
+	}
+	
+	while (!CommitQueue.IsEmpty() && CommitQueue.Last() == nullptr)
+	{
+		CommitQueue.Pop(EAllowShrinking::No);
+	}
+
+	if (CommitQueue.IsEmpty())
+	{
+		CommitScratch.Reset();
+	}
+}
+
+void FHistory::Link(TConstArrayView<FOpsPtr> Committed)
+{
+	FOptionalDeclId Last;
+	for (const FOpsPtr& Ops : Committed)
+	{
+		if (Ops && Ops->Owner != Last)
+		{
+			Link(Ops->Owner);
+			Last = Ops->Owner;
+		}
+	}
+}
+
+template<class T>
+static const T& GetDecomposition(const FVersionUpgrades& Version, uint16 DecompIdx)
+{
+	check(DecompIdx < Version.NumDecompositions);
+	return Version.At<T>(Version.GetDecompOffsets()[DecompIdx]);
+}
+
+void FHistory::Link(FDeclId Compiled)
+{
+	FStructUpgrades& Upgrades = *CompiledUpgrades.FindChecked(Compiled).Get();
+	for (const FVersionUpgrades* Version : ViewVersions(Upgrades))
+	{
+		for (FMemberMatch Match : ViewMatches(*Version))
+		{
+			if (Match.Decomposes())
+			{
+				if (Match.Type.IsStruct())
+				{
+					Link(const_cast<FStructDecomposition&>(Version->At<FStructDecomposition>(Match.Idx.Decomp)));
+				}
+			}
+		}
+	}
+}
+
+static bool CompareParameters(FLooseParameter A, FLooseParameter B)
+{
+	return A.Name == B.Name && A.Type.Compare(B.Type);
+}
+
+// Temp checks that latest declaration matches the input order exactly until overriding is implemented
+static void TemporaryInputSchemaCheck(const FLooseSchema& Declared, const FLooseSchema& Input)
+{
+	check(Declared.Version == Input.Version);
+	check(Declared.NumMembers >= Input.NumMembers);
+
+	TConstArrayView<FLooseParameter> Expecteds = Declared.GetMembers();
+	for (FLooseParameter Actual : Input.GetMembers())	
+	{
+		const FLooseParameter* Expected = Expecteds.FindByPredicate([=](FLooseParameter Param) { return CompareParameters(Actual, Param); });
+		check(Expected);
+		Expecteds.RightChopInline(Expected == Expecteds.GetData());
+	}
+}
+
+void FHistory::Link(FStructDecomposition& Dependency)
+{
+	if (const FStructUpgradesPtr* Upgrades = CompiledUpgrades.Find(Dependency.Id))
+	{
+		// todo: override Upgrades with Recipient 1
+
+		if (const FLooseSchema* DeclaredSchema = LatestSchemas.Obtain(Dependency.Id, Latest, LooseIds.Get(Dependency.Id)))
+		{
+			TemporaryInputSchemaCheck(*DeclaredSchema, *Dependency.InputSchema);
+		}
+		
+		Dependency.Upgrades = Upgrades->Get();
+	}
+}
+
+////////////////////////////////////////////////////////////////////////////////////////////////
+
+FHistoryPtr MakeHistory(IDeclarations& Latest, FLiteralIndexerBase& Ids)
+{
+	return FHistoryPtr(new FHistory{Latest, {}, Ids, Latest});
+}
+
+void FHistoryDeleter::operator()(FHistory* Ptr) const
+{
+	delete Ptr;
+}
+
+void EraseChroniclesOf(FHistory& History, FStructId Owner)
+{
+	History.Erase(Owner);
+}
+
+////////////////////////////////////////////////////////////////////////////////////////////////
+
+IOps* FChronicles::CreateVersion(FTypename Struct, uint16 Ver)
+{
+	FDeclId Id(IndexStruct(History.Ids, Struct));
+	return new (History.CommitScratch) FOps(Id, Ver, History.Ids, History.LooseIds);
+}
+
+void FChronicles::CommitVersion(IOps* InOps)
+{
+	FOps* Ops = static_cast<FOps*>(InOps);
+	checkf(uint16(PriorVersion + uint16(1)) == Ops->Version, TEXT("Commit versions in contiguous order")); // Overflows first time
+	checkf(!bFinalVersion, TEXT("Final version already committed, version %d failed to Bump()"), PriorVersion);
+	PriorVersion = Ops->Version;
+	bFinalVersion = Ops->Redirect || !Ops->NextVersion;
+
+	History.CommitQueue.Emplace(Ops);
+}
+
+}} // namespace PlainProps::Upgrade
\ No newline at end of file
diff --git a/Engine/Plugins/Experimental/PlainProps/Source/Public/PlainPropsLooseTypes.h b/Engine/Plugins/Experimental/PlainProps/Source/Public/PlainPropsLooseTypes.h
new file mode 100644
--- /dev/null
+++ b/Engine/Plugins/Experimental/PlainProps/Source/Public/PlainPropsLooseTypes.h
@@ -0,0 +1,246 @@
+// Copyright Epic Games, Inc. All Rights Reserved.
+
+#pragma once
+
+#include "PlainPropsTypes.h"
+#include "Memory/MemoryFwd.h"
+#include "Misc/Optional.h"
+
+namespace PlainProps
+{
+class FLooseMember;
+struct FLooseNestedRange;
+struct FLooseValue;
+
+////////////////////////////////////////////////////////////////////////////////////////////////
+
+enum class EInnermostKind : uint8 { Struct, Enum, Bool, S8, S16, S32, S64, U8, U16, U32, U64, F32, F64, UTF8, UTF16, UTF32 };
+
+struct FInnermostStructType
+{
+	EInnermostKind			_ : 6;
+	uint8					IsDynamic : 1;
+	uint8					IsSuper : 1; // Needed given FOptionalMemberId-indexing?
+};
+
+struct FInnermostEnumType
+{
+	EInnermostKind			_ : 6;
+	uint8					FlagMode : 1;
+	uint8					Unused : 1;
+};
+
+union FInnermostType
+{
+	EInnermostKind			Kind : 6;
+	FInnermostStructType	Struct;
+	FInnermostEnumType		Enum;
+};
+
+inline static constexpr FInnermostType InnermostFlatEnum = { .Enum = { ._ = EInnermostKind::Enum, .FlagMode = 0 } };
+inline static constexpr FInnermostType InnermostFlagEnum = { .Enum = { ._ = EInnermostKind::Enum, .FlagMode = 1 } };
+inline static constexpr FInnermostType InnermostStaticStruct = { .Struct = { ._ = EInnermostKind::Struct, .IsDynamic = 0 } };
+inline static constexpr FInnermostType InnermostDynamicStruct = { .Struct = { ._ = EInnermostKind::Struct, .IsDynamic = 1 } };
+
+template<typename T>
+constexpr FInnermostType IllegalInnermost()
+{
+	static_assert(!sizeof(T), "Unsupported leaf type");
+	return InnermostFlatEnum;
+}
+
+template<Arithmetic T>
+inline constexpr FInnermostType ReflectInnermost = IllegalInnermost<T>();
+
+template<> inline constexpr FInnermostType ReflectInnermost<bool>		= { .Kind = EInnermostKind::Bool };
+template<> inline constexpr FInnermostType ReflectInnermost<int8>		= { .Kind = EInnermostKind::S8 };
+template<> inline constexpr FInnermostType ReflectInnermost<int16>		= { .Kind = EInnermostKind::S16 };
+template<> inline constexpr FInnermostType ReflectInnermost<int32>		= { .Kind = EInnermostKind::S32 };
+template<> inline constexpr FInnermostType ReflectInnermost<int64>		= { .Kind = EInnermostKind::S64 };
+template<> inline constexpr FInnermostType ReflectInnermost<uint8>		= { .Kind = EInnermostKind::U8 };
+template<> inline constexpr FInnermostType ReflectInnermost<uint16>		= { .Kind = EInnermostKind::U16 };
+template<> inline constexpr FInnermostType ReflectInnermost<uint32>		= { .Kind = EInnermostKind::U32 };
+template<> inline constexpr FInnermostType ReflectInnermost<uint64>		= { .Kind = EInnermostKind::U64 };
+template<> inline constexpr FInnermostType ReflectInnermost<float>		= { .Kind = EInnermostKind::F32 };
+template<> inline constexpr FInnermostType ReflectInnermost<double>		= { .Kind = EInnermostKind::F64 };
+template<> inline constexpr FInnermostType ReflectInnermost<char8_t>	= { .Kind = EInnermostKind::UTF8 };
+template<> inline constexpr FInnermostType ReflectInnermost<char16_t>	= { .Kind = EInnermostKind::UTF16 };
+template<> inline constexpr FInnermostType ReflectInnermost<char32_t>	= { .Kind = EInnermostKind::UTF32 };
+
+//template<Enumeration>
+//inline constexpr FInnermostType ReflectInnermost = { .Kind = EInnermostKind::Enum };
+//
+//template<typename T>
+//inline constexpr FInnermostType ReflectInnermost requires (std::is_class_v<T>) = { .Kind = EInnermostKind::Struct };
+
+////////////////////////////////////////////////////////////////////////////////////////////////
+
+template<Arithmetic T>	uint64	ValueCast(T Value)						{ return static_cast<std::make_unsigned_t<T>>(Value); }
+template<> inline		uint64	ValueCast(float Value)					{ return BitCast<uint32>(Value); }
+template<> inline		uint64	ValueCast(double Value)					{ return BitCast<uint64>(Value); }
+
+template<Arithmetic T>	T		ArithmeticCast(uint64 Value)			{ return static_cast<T>(IntCastChecked<std::make_unsigned_t<T>>(Value)); }
+template<> inline		bool	ArithmeticCast<bool>(uint64 Value)		{ return IntCastChecked<bool>(Value); }
+template<> inline		float	ArithmeticCast<float>(uint64 Value)		{ return BitCast<float>(IntCastChecked<uint32>(Value)); }
+template<> inline		double	ArithmeticCast<double>(uint64 Value)	{ return BitCast<double>(Value); }
+
+////////////////////////////////////////////////////////////////////////////////////////////////
+
+static constexpr uint32 LooseRangeBits = 6;
+static constexpr uint8 MaxRangeNesting = (uint8(1) << LooseRangeBits) - 1; // Todo: Enforce on build side too
+
+struct FLooseType
+{
+	uint8				OptionalParameter : 1 = 0;
+	uint8				NumRanges : LooseRangeBits + /* reserved */ 1 = 0;
+	FInnermostType		InnermostType = InnermostStaticStruct;
+	uint16				InnermostVersion = 0;
+	FOptionalInnerId	InnermostId;
+
+	bool				IsRange() const					{ return NumRanges != 0; }
+	bool				IsStruct() const				{ return NumRanges == 0 && InnermostType.Kind == EInnermostKind::Struct; }
+	bool				IsLeaf() const					{ return NumRanges == 0 && InnermostType.Kind != EInnermostKind::Struct; }
+	bool				IsEnum() const					{ return NumRanges == 0 && InnermostType.Kind == EInnermostKind::Enum; }
+
+	bool				Compare(FLooseType O) const		{ return WithoutMetadata().AsInt() == O.WithoutMetadata().AsInt(); }
+	bool				CompareAll(FLooseType O) const	{ return AsInt() == O.AsInt(); }
+	FLooseType			WithoutMetadata() const			{ return {0, NumRanges, InnermostType, InnermostVersion, InnermostId}; }
+private:
+	uint64				AsInt() const					{ return BitCast<uint64>(*this); }
+};
+
+template<Arithmetic T>
+inline constexpr FLooseType ReflectLoose = { 0, 0, ReflectInnermost<T>, 0, NoId };
+
+//////////////////////////////////////////////////////////////////////////////////////////////////
+
+// Per enum-indexed enumerator name
+struct FEnumeratorId
+{
+	uint16 Idx = 0xFFFF;
+
+	bool operator==(FEnumeratorId O) const { return Idx == O.Idx; }
+	friend uint32 GetTypeHash(FEnumeratorId Id) { return Id.Idx; };
+};
+
+// Set of FEnumeratorId, each matching an enum flag name
+struct FEnumeratorIdSet
+{
+	uint64				Bits = 0;		// <= 64 FEnumeratorId flag names are inlined
+};
+
+union FEnumeratorData
+{
+	FEnumeratorId		Flat;	// Loose enumerator name (without integer value)
+	FEnumeratorIdSet	Flags;	// Loose set of enumerator flag names without integer value
+};
+
+struct FEnumeratorInput
+{
+	uint16				Index;
+};
+
+struct FEnumeratorInputSet
+{
+	uint64				Bits;
+};
+
+// Temporary format used to access individual enum members during Op execution
+union FEnumIndexData
+{
+	FEnumeratorInput	Flat;	// Enumerator index using NameEnumerators order
+	FEnumeratorInputSet	Flags;	// Bitset of flag indices using NameFlags order
+};
+
+union FEnumLooseData
+{
+	FEnumeratorData		Enumerators;	// FEnumeratorId::Idx order from FEnumeratorIndexer
+	FEnumIndexData		Indices;		// Input order from NameEnumerators/NameFlags
+};
+
+union FEnumData
+{
+	uint64				Source = 0;	// Integer source value matching the innermost enum schema
+	FEnumLooseData		Loose;
+};
+
+union FStructRange
+{
+	const void*			Source = nullptr; // Skippable slice
+	const FLooseValue*	Loose;
+};
+
+union FRangeData
+{
+	const void*					Leaves = nullptr;
+	FStructRange				Structs;
+	const FLooseNestedRange*	Ranges; 
+};
+
+struct FLooseNestedRange
+{
+	uint64				Num = 0;
+	FRangeData			Data;
+};
+
+union FStructData
+{
+	const uint8*		Source = nullptr;	// Skippable slice from a source FStructView
+	FLooseMember*		Loose;
+};
+
+union FLooseData
+{
+	uint64				Arithmetic = 0;
+	FEnumData			Enum;
+	FRangeData			Range;
+	FStructData			Struct;
+};
+
+// FLooseData metadata
+struct FLooseMetadata
+{
+	uint8				IsSet = 0;
+	uint8				IsSourceRange : 1 = 0;
+	uint8				IsFullyUpgraded : 1 = 0;
+	uint16				Name = 0;					// FParameterId / FEnumeratorId
+	FOptionalSchemaId	InnermostSchema;			// Set if innermost is source struct/enum
+	uint64				Num = 0;					// Range items, loose struct members or inner members / enumerators during op execution
+};
+
+struct FLooseValue
+{
+	FLooseMetadata		Meta;
+	FLooseData			Data;
+
+	explicit operator	bool() const { return !!Meta.IsSet; }
+
+	void				SetNamed(FLooseValue In)	// Assign but keep name
+	{
+		In.Meta.Name = Meta.Name;
+		*this = In;
+	}
+};
+
+//////////////////////////////////////////////////////////////////////////////////////////////////
+
+// Type and named optional value
+class FLooseMember
+{
+public:
+	FLooseMember(FLooseType T, FLooseValue V) : Type(T), Value(V) {}
+	
+	explicit			operator bool() const			{ return !!Value.Meta.IsSet; }
+	FLooseType			GetType() const					{ return Type; }
+	FLooseValue			GetValue() const				{ return Value; }
+	void				Rename(uint16 ParameterId)		{ Value.Meta.Name = ParameterId; }
+	void				SetTypechecked(FLooseValue V)	{ Value = V; } 
+	FLooseValue&		EditTypechecked()				{ return Value; } 
+	bool				IsOptional() const				{ return Type.OptionalParameter; }
+
+protected:
+	FLooseType			Type;
+	FLooseValue			Value;
+};
+
+} // namespace PlainProps
\ No newline at end of file
diff --git a/Engine/Plugins/Experimental/PlainProps/Source/Public/PlainPropsUpgrade.h b/Engine/Plugins/Experimental/PlainProps/Source/Public/PlainPropsUpgrade.h
new file mode 100644
--- /dev/null
+++ b/Engine/Plugins/Experimental/PlainProps/Source/Public/PlainPropsUpgrade.h
@@ -0,0 +1,650 @@
+// Copyright Epic Games, Inc. All Rights Reserved.
+
+#pragma once
+
+#include "PlainPropsTypes.h"
+#include "PlainPropsLooseTypes.h"
+#include "Containers/BitArray.h"
+#include "Misc/Optional.h"
+#include "Templates/UniquePtr.h"
+
+
+namespace PlainProps
+{
+struct FBuiltStruct;
+class FCustomBindings;
+class FIdIndexerBase;
+class FLiteralIndexerBase;
+class FSchemaBatchId;
+struct FSchemaBatch;
+class FSchemaBindings;
+class FScratchAllocator;
+struct FStructDeclaration;
+struct FStructView;
+struct FUpgradeBindings;
+class FWriter;
+struct IDeclarations;
+
+namespace Upgrade
+{
+class FChronicler;
+class FContext;	
+class FEnumChronicler;
+struct FHistory;
+class FIn;
+class FInEnum;
+class FInEnumerator;
+class FInEnumFlags;
+class FInStruct;
+struct IOps;
+class FOut;
+class FOutEnum;
+class FOutEnumerator;
+class FOutEnumFlags;
+struct FOutParameter;
+class FOutStruct;
+class FUpgrades;
+template<uint32 N> class TOuts;
+
+using FIns = TConstArrayView<FIn>;
+using FOuts = TArrayView<FOut>;
+using Op = void (*)(FIns, FOuts, FContext); 
+using AddOps = void (*)(IOps& Out); 
+
+////////////////////////////////////////////////////////////////////////////////////////////////
+
+// Mechanism for compiling out old upgrade code and literals
+template<Enumeration ReleaseEnum> 
+inline constexpr ReleaseEnum OldestRelevant = {};
+
+template<auto Release> 
+inline constexpr bool Relevant = Release >= OldestRelevant<decltype(Release)>;
+
+////////////////////////////////////////////////////////////////////////////////////////////////
+
+enum class EAccess : uint8 { WholeValue, AllMembers, SomeMembers }; // todo: could allow leaving taken members too. could have another flag to leave old member
+
+struct FInnermostHandle
+{
+	FOptionalInnerId	Id;
+	uint16				Version;
+	uint8				Slice;
+	EAccess				Access;
+};
+
+struct FParameter
+{
+	const char*			Name;
+	uint32				NameLen;
+	uint8				NumRanges;
+	bool				bOptional;
+	FInnermostType		InnermostType;
+	FInnermostHandle	InnermostHandle;
+
+	FParameter operator~() const { return {Name, NameLen, NumRanges, true, InnermostType, InnermostHandle}; }
+};
+
+template<Arithmetic T>
+constexpr FParameter MakeParam(const char* Name, std::size_t Len)
+{
+	return { Name, static_cast<uint32>(Len), 0, false, ReflectInnermost<T>, {} };
+}
+
+////////////////////////////////////////////////////////////////////////////////////////////////
+
+constexpr FParameter operator ""_B1(const char* Name, std::size_t Len)	{ return MakeParam<bool>(Name, Len); }
+constexpr FParameter operator ""_S32(const char* Name, std::size_t Len)	{ return MakeParam<int32>(Name, Len); }
+constexpr FParameter operator ""_U32(const char* Name, std::size_t Len)	{ return MakeParam<uint32>(Name, Len); }
+constexpr FParameter operator ""_U64(const char* Name, std::size_t Len)	{ return MakeParam<uint64>(Name, Len); }
+constexpr FParameter operator ""_F32(const char* Name, std::size_t Len)	{ return MakeParam<float>(Name, Len); }
+constexpr FParameter operator ""_F64(const char* Name, std::size_t Len)	{ return MakeParam<double>(Name, Len); }
+
+////////////////////////////////////////////////////////////////////////////////////////////////
+
+struct FMemberLiteral
+{
+	const char*			Str;
+	std::size_t			Len;
+
+	inline constexpr FParameter MakeParam(uint8 NumRanges, FInnermostType InnermostType, FInnermostHandle InnermostHandle) const
+	{
+		return { Str, static_cast<uint32>(Len), NumRanges, false, InnermostType, InnermostHandle };
+	}
+};
+
+////////////////////////////////////////////////////////////////////////////////////////////////
+
+struct FEnumeratorNames
+{
+	FEnumId			Id;
+	uint8			Slice;
+};
+
+struct FEnumFlagNames
+{
+	FEnumId			Id;
+	uint8			Slice;
+};
+
+struct FEnumLiteral : FMemberLiteral
+{
+	inline FParameter operator()(FEnumId Id) const
+	{
+		return MakeParam(0, InnermostFlatEnum, {FInnerId(Id), 0, 0, EAccess::WholeValue});
+	}
+
+	inline FParameter operator()(FEnumeratorNames Enumerators, EAccess Access = EAccess::AllMembers) const
+	{
+		return MakeParam(0, InnermostFlatEnum, {FInnerId(Enumerators.Id), 0, Enumerators.Slice, Access});
+	}
+
+	inline FParameter operator()(FEnumFlagNames Flags, EAccess Access = EAccess::AllMembers) const
+	{
+		return MakeParam(0, InnermostFlagEnum, {FInnerId(Flags.Id), 0, Flags.Slice, Access});
+	}
+};
+inline constexpr FEnumLiteral operator ""_Enum(const char* Name, std::size_t Len)		{ return { {Name, Len} }; }
+
+////////////////////////////////////////////////////////////////////////////////////////////////
+
+struct FStructMemberNames
+{
+	FStructId		Id;
+	uint8			Slice;
+};
+
+struct FStructLiteral : FMemberLiteral
+{
+	inline FParameter operator()(FStructId Id, uint16 Version = 0) const
+	{
+		return MakeParam(0, InnermostStaticStruct, {FInnerId(Id), Version, 0, EAccess::WholeValue});
+	}
+
+	inline FParameter operator()(FStructMemberNames Members, uint16 Version = 0, EAccess Access = EAccess::AllMembers) const
+	{
+		return MakeParam(0, InnermostStaticStruct, {FInnerId(Members.Id), Version, Members.Slice, Access});
+	}
+};
+inline constexpr FStructLiteral operator ""_Struct(const char* Name, std::size_t Len)	{ return { Name, Len }; }
+inline constexpr FParameter operator ""_AnyStruct(const char* Name, std::size_t Len)	{ return { Name, static_cast<uint32>(Len), 0, false, InnermostDynamicStruct, {} }; }
+
+////////////////////////////////////////////////////////////////////////////////////////////////
+
+struct FRangeLiteral : FMemberLiteral
+{
+	uint8				NumRanges;
+
+	template<Arithmetic T>
+	inline FParameter Of() const
+	{
+		return Make(ReflectInnermost<T>);
+	}
+
+	inline FParameter Of(FEnumId Id) const
+	{
+		return MakeParam(NumRanges, InnermostFlatEnum, {FInnerId(Id), 0, 0, EAccess::WholeValue});
+	}
+
+	inline FParameter Of(FEnumFlagNames Flags, EAccess Access = EAccess::AllMembers) const
+	{
+		return MakeParam(NumRanges, InnermostFlagEnum, {FInnerId(Flags.Id), 0, Flags.Slice, Access});
+	}
+
+	inline FParameter Of(FStructId Id, uint16 Version = 0) const
+	{
+		return MakeParam(NumRanges, InnermostStaticStruct, {FInnerId(Id), Version, 0, EAccess::WholeValue});
+	}
+
+	inline FParameter Of(FStructMemberNames Members, uint16 Version = 0, EAccess Access = EAccess::AllMembers) const
+	{
+		return MakeParam(NumRanges, InnermostStaticStruct, {FInnerId(Members.Id), Version, Members.Slice, Access});
+	}
+
+	inline FParameter OfAny() const
+	{
+		return MakeParam(NumRanges, InnermostDynamicStruct, {NoId, 0, 0, EAccess::WholeValue});
+	}
+
+	inline FRangeLiteral Nest(uint8 NumNested = 1) const
+	{
+		check(NumRanges == 1);
+		check(NumNested > 0 && NumNested < MaxRangeNesting);
+		return { {Str, Len}, static_cast<uint8>(NumRanges + NumNested) };
+	}
+};
+consteval FRangeLiteral operator ""_Range(const char* Name, std::size_t Len) { return { {Name, Len}, 1 }; }
+
+////////////////////////////////////////////////////////////////////////////////////////////////
+
+// Used to access inner struct members
+template<uint32 N>
+class TIns
+{
+public:
+	TIns(const FIn* InIns) : Ins(InIns) { static_assert(N > 1); }
+	const FIn& operator[](uint32 Idx) const;
+private:
+	const FIn* Ins;
+};
+
+class FIn : protected FLooseMember
+{
+public:
+	using FLooseMember::FLooseMember;
+	using FLooseMember::operator bool;
+
+	bool							Bool() const { return As<bool>(); }
+	int32							S32() const { return As<int32>(); }
+	uint32							U32() const { return As<uint32>(); }
+	uint64							U64() const { return As<uint64>(); }
+	float							F32() const { return As<float>(); }
+	PLAINPROPS_API FInEnum			Enum() const; 
+	PLAINPROPS_API FInEnumerator	Enumerator() const;
+	PLAINPROPS_API FInEnumFlags		EnumFlags() const;
+	PLAINPROPS_API FInStruct		Struct() const;
+
+
+	TOptional<bool>					OptBool() const { return AsOpt<bool>(); }
+	TOptional<int32>				OptS32() const { return AsOpt<int32>(); }
+	TOptional<uint32>				OptU32() const { return AsOpt<uint32>(); }
+	TOptional<uint64>				OptU64() const { return AsOpt<uint64>(); }
+	TOptional<float>				OptF32() const { return AsOpt<float>(); }
+	PLAINPROPS_API FInStruct		OptStruct() const;
+
+	template<uint32 N> TIns<N>		StructMembers() const { return TIns<N>(AsStructMembers(N)); }
+
+private:	
+	template<Arithmetic T>
+	T								As() const
+	{
+		check(Type.Compare(ReflectLoose<T>));
+		check(Value);
+		return ArithmeticCast<T>(Value.Data.Arithmetic);
+	}
+
+	template<Arithmetic T>
+	TOptional<T>					AsOpt() const
+	{
+		check(Type.Compare(ReflectLoose<T>));
+		check(IsOptional());
+		return Value ? TOptional<T>(ArithmeticCast<T>(Value.Data.Arithmetic)) : NullOpt;
+	}
+
+	PLAINPROPS_API const FIn*		AsStructMembers(uint32 N) const;
+};
+
+template<uint32 N>
+const FIn& TIns<N>::operator[](uint32 Idx) const
+{
+	check(Idx < N);
+	return Ins[Idx];
+}
+
+// Opaque struct value, assignable to a type-matching FOutStruct
+class FInStruct : protected FIn
+{
+	friend FIn; friend FOutStruct;
+};
+
+// Opaque enum value, assignable to a type-matching FOutEnum
+class FInEnum : protected FIn
+{
+	friend FIn; friend FOutEnum;
+};
+
+// Enum value represented as enumerator index in NameEnumerators() order
+// * Conceptually a string
+// * Not assignable to FOutEnumerator, which might use different NameEnumerators() indices
+// * Not convertible to FInEnum nor assignable to FOutEnum either
+class FInEnumerator
+{
+	friend FIn;
+	uint8 Index;
+public:
+	explicit FInEnumerator(uint8 Idx) : Index(Idx) {}
+
+	uint8			GetIndex() const { return Index;}
+};
+
+// Bitset of enumerator indices in NameEnumFlags() order
+// * Conceptually a string set
+// * Not assignable to FOutEnumerator, which might use different NameEnumerators() indices
+class FInEnumFlags
+{
+	uint64 Indices;
+public:
+	explicit FInEnumFlags(uint64 IndexFlags) : Indices(IndexFlags) {}
+	uint64			GetIndices() const;
+	bool			IsSet(uint8 Index) const;
+};
+
+////////////////////////////////////////////////////////////////////////////////////////////////
+
+template <Arithmetic T>
+class TOut
+{
+	FLooseValue& Out;
+public:
+	TOut(FLooseType Type, FLooseValue& Value UE_LIFETIMEBOUND)
+	: Out(Value)
+	{
+		check(Type.Compare(ReflectLoose<T>));
+	}
+
+	void operator=(T In)
+	{
+		Out.Meta.IsSet = 1;
+		Out.Data.Arithmetic = ValueCast(In);
+	}
+
+	void operator=(TOptional<T> In)
+	{
+		Out.Meta.IsSet = !!In;
+		Out.Data.Arithmetic =  ValueCast(In.Get(T{}));
+	}
+};
+
+// TOutRange<Arithmetic> helper
+template <Arithmetic T>
+class TNormalizer
+{
+	T& Out;
+public:
+	TNormalizer(T& O) : Out(O) {}
+	void operator=(T In) { Out = Normalize(In);	}
+};
+
+template <Arithmetic T>
+class TOutRange
+{
+	TArrayView64<T> Out;
+public:
+	TOutRange(uint64 Num, T* Data) : Out(Data, Num) {}
+	TNormalizer<T>&	operator[](int64 Idx) { return Out[Idx]; }
+};
+
+// TOutRange<bool> helper
+class FBitSetter
+{
+	uint8&				Out;
+	const uint8			Mask;
+public:
+	FBitSetter(uint8& Byte, uint8 Idx) : Out(Byte), Mask(uint8(1) << Idx) {}
+
+	void operator=(bool bValue)
+	{
+		if (bValue)
+		{
+			Out |= Mask;
+		}
+		else
+		{
+			Out &= ~Mask;
+		}	
+	}
+};
+
+template<>
+class TOutRange<bool>
+{
+	const uint64					Num;
+	uint8*							Out;
+public:
+	TOutRange(uint64 N, uint8* Bits) : Num(N), Out(Bits) {}
+	FBitSetter operator[](uint64 Idx) { check(Idx < Num); return FBitSetter(Out[Idx/8], Idx%8); }
+};
+
+class FOutStructs
+{
+	const FDeclId					Id;
+	const uint16					Version;
+	TArrayView<FLooseValue>			Out;
+public:
+	FOutStructs(FDeclId I, uint16 V, TArrayView<FLooseValue> O) : Id(I), Version(V), Out(O) {}
+};
+
+class FOutRange
+{
+public:
+	FOutRange(FLooseType T, uint64 N, void* D) : Type(T), Num(N), Data(D) {}
+	
+	TOutRange<int32>				S32s() { return Of<int32>(); }
+	TOutRange<uint32>				U32s() { return Of<uint32>(); }
+	TOutRange<uint64>				U64s() { return Of<uint64>(); }
+	TOutRange<float>				F32s() { return Of<float>(); }
+	FOutStructs						Structs();// { check(Type.IsStruct()); return Items; }
+	FOutStructs						DynamicStructs(FDeclId Id, uint16 Version);// { check(Type.IsStruct() && Type.InnermostType.) }
+
+private:
+	const FLooseType				Type;
+	uint64							Num;
+	void*							Data;
+
+	template<Arithmetic T>
+	TOutRange<T>					Of()
+	{
+		check(Type.Compare(ReflectLoose<T>));
+		return TOutRange<T>(Num, static_cast<T*>(Data));
+	}
+};
+
+class FOut : protected FLooseMember
+{
+public:
+	FOut(FOut&&) = default;
+	FOut(const FOut&) = delete;
+	FOut(FOutParameter Parameter);
+	FOut(FOutParameter Parameter, FOut* StructMembers);
+	FOut& operator=(FOut&&) = default;
+	FOut& operator=(const FOut&) = delete;
+
+	TOut<int32>						S32()		{ return { Type, Value }; }
+	TOut<uint32>					U32()		{ return { Type, Value }; }
+	TOut<uint64>					U64()		{ return { Type, Value }; }
+	TOut<float>						F32()		{ return { Type, Value }; }
+	
+	PLAINPROPS_API FOutEnum&		Enum();
+	PLAINPROPS_API FOutEnumerator&	Enumerator();
+	PLAINPROPS_API FOutEnumFlags&	EnumFlags();
+	PLAINPROPS_API FOutStruct&		Struct();
+	template<uint32 N> TOuts<N>		StructMembers() { return AsStructMembers(N); }
+
+	PLAINPROPS_API FOutRange		NewRange(uint64 Num); // Allocate zeroed range (note: copy Num down for e.g. range of struct members)
+
+private:
+	PLAINPROPS_API FOut*			AsStructMembers(uint32 N);
+
+	friend const FLooseMember&		ReadOutput(const FOut& Out);
+	friend TArrayView<FLooseMember>	EditOutputs(TArrayView<FOut> Outs);
+};
+
+class FOutEnumFlag
+{
+	FLooseValue& Out;
+	const uint64 Mask;
+public:
+	FOutEnumFlag(FLooseValue& O, uint32 Idx) : Out(O), Mask(uint64(1) << Idx) { check(Idx < Out.Meta.Num); }
+	PLAINPROPS_API void operator=(bool bValue);
+	PLAINPROPS_API void operator=(TOptional<bool> bValue);
+};
+
+class FOutEnumFlags : protected FOut
+{
+public:
+	FOutEnumFlag operator[](uint32 Idx) { return FOutEnumFlag(/* out */ Value, Idx); }
+};
+
+class FOutEnum : protected FOut
+{
+public:
+	PLAINPROPS_API void operator=(FInEnum In);
+};
+
+class FOutStruct : protected FOut
+{
+public:
+	PLAINPROPS_API void operator=(FInStruct In);
+};
+
+template<uint32 N>
+class TOuts
+{
+public:
+	TOuts(FOut* InOuts) : Outs(InOuts) {}
+	FOut& operator[](uint32 Idx) const { check(Idx < N); return Outs[Idx]; }
+private:
+	FOut* const Outs;
+};
+
+////////////////////////////////////////////////////////////////////////////////////////////////
+
+// Only for logging errors and warnings, todo
+class FContext {};
+
+////////////////////////////////////////////////////////////////////////////////////////////////
+
+// Unindexed typename
+struct FTypename
+{
+	const char*		Scope;
+	const char*		Name;
+};
+
+// All upgrade ops for a particular version of a particular struct
+struct IOps
+{
+	virtual ~IOps() {}
+	
+	virtual FEnumId				NameEnum(FTypename Enum) = 0;
+	virtual FEnumeratorNames	NameEnumerators(FEnumId Enum, std::initializer_list<const char*> Enumerators) = 0;
+	inline  FEnumeratorNames	NameEnumerators(FTypename Enum, std::initializer_list<const char*> Enumerators)		{ return NameEnumerators(NameEnum(Enum), Enumerators); }
+	virtual FEnumFlagNames		NameFlags(FEnumId Enum, std::initializer_list<const char*> Flags) = 0;
+	inline  FEnumFlagNames		NameFlags(FTypename Enum, std::initializer_list<const char*> Flags)					{ return NameFlags(NameEnum(Enum), Flags); }
+	virtual FStructId			NameStruct(FTypename Struct) = 0;
+	virtual FStructMemberNames	NameMembers(FStructId Struct, TConstArrayView<FParameter> Members) = 0;
+	inline  FStructMemberNames	NameMembers(FTypename Struct, TConstArrayView<FParameter> Members)					{ return NameMembers(NameStruct(Struct), Members); }
+
+	virtual void				Add(TConstArrayView<FParameter> Ins, TConstArrayView<FParameter> Outs, Op Upgrade) = 0;
+
+	template<uint32 N, uint32 M>
+	void						Add(FParameter (&&Ins)[N], FParameter (&&Outs)[M], Op Upgrade)
+	{
+		Add(MakeArrayView(Ins), MakeArrayView(Outs), Upgrade);
+	}
+	
+	// Only one version bump or rename is allowed and must happen last
+	virtual void				Bump(uint16 NextVersion) = 0;
+
+	// Only one version bump or rename is allowed and must happen last
+	virtual void				Rename(FTypename Redirect, uint16 ToVersion) = 0;
+};
+
+////////////////////////////////////////////////////////////////////////////////////////////////
+
+class FChronicles
+{
+	FHistory&					History;
+	FStructId					Owner;
+	uint16						PriorVersion = 0xFFFF;
+	bool						bFinalVersion = false;
+
+public:
+	FChronicles(FHistory& Out, FStructId InOwner) : History(Out), Owner(InOwner) {}
+
+	PLAINPROPS_API IOps*	CreateVersion(FTypename Name, uint16 Version);
+	PLAINPROPS_API void		CommitVersion(IOps* Version);
+};
+
+// Records relevant ops for a particular struct version in its destructor
+// Helps compile out irrelevant literals and ops
+class FChronicler
+{
+public:
+	FORCEINLINE FChronicler(FChronicles Out, const char* Scope, const char* Name, uint16 Ver = 0) : Name{Scope, Name}, Version(Ver), Chronicles(Out) {}
+	FORCEINLINE ~FChronicler()
+	{
+		if (Out)
+		{
+			Chronicles.CommitVersion(/* move */ Out);
+		}	
+	}
+	
+	template<auto Release>
+	FORCEINLINE void In(AddOps Add)
+	{
+		if constexpr (Relevant<Release>)
+		{
+			Out = Out ? Out : Chronicles.CreateVersion(Name, Version);
+			Add(/* out */ *Out);
+		}
+	}
+
+	template<auto Release>
+	FORCEINLINE void Rename(FTypename To, uint16 ToVersion = 0)
+	{
+		if constexpr (Relevant<Release>)
+		{
+			Out = Out ? Out : Chronicles.CreateVersion(Name, Version);
+			Out->Rename(To, ToVersion);
+		}
+	}
+
+	FTypename GetName() const { return Name; }
+
+private:
+	const FTypename		Name;
+	const uint16		Version;
+	FChronicles			Chronicles;
+	IOps*				Out = nullptr;
+};
+
+// todo
+class FEnumChronicler
+{};
+
+////////////////////////////////////////////////////////////////////////////////////////////////
+
+// Enables overloading void Record(TChroniclesOf<FMyStruct> Out) 
+template <typename T>
+class TChroniclesOf : public FChronicles {};
+
+template<typename T>
+void Record(TChroniclesOf<T>) {}
+
+PLAINPROPS_API void EraseChroniclesOf(FHistory& History, FStructId Owner);
+
+////////////////////////////////////////////////////////////////////////////////////////////////
+
+struct FHistoryDeleter
+{
+	PLAINPROPS_API void operator()(FHistory* Ptr) const;
+};
+using FHistoryPtr = TUniquePtr<FHistory, FHistoryDeleter>;
+PLAINPROPS_API FHistoryPtr MakeHistory(IDeclarations& Latest, FLiteralIndexerBase& Ids);
+
+////////////////////////////////////////////////////////////////////////////////////////////////
+
+struct FUpgradedInstance
+{
+	FDeclId					Type;
+	const FBuiltStruct*		Value;
+};
+
+class FBatchUpgrader
+{
+public:
+	PLAINPROPS_API FBatchUpgrader(const FSchemaBatch& InBatch UE_LIFETIMEBOUND, TConstArrayView<FInnerId> InRuntimeIds UE_LIFETIMEBOUND, FIdIndexerBase& Ids);
+
+	// @param Types must return declarations that outlive this
+	PLAINPROPS_API bool							MatchSchemas(const IDeclarations& Types);
+
+	// todo: add unknown members api and error handler
+	PLAINPROPS_API TArray<FUpgradedInstance>	Upgrade(TConstArrayView<FStructView> Structs, FHistory& Upgrades, FScratchAllocator& Scratch);
+private:
+	const FSchemaBatch&							Batch;
+	TConstArrayView<FInnerId> 					RuntimeIds;
+	TBitArray<>									MatchingSchemas;
+	TArray<const FStructDeclaration*>			Declarations;
+};
+
+}} // namespace PlainProps::Upgrade
\ No newline at end of file
diff --git a/Engine/Plugins/Experimental/PlainProps/Source/Tests/PlainPropsUpgradeTests.cpp b/Engine/Plugins/Experimental/PlainProps/Source/Tests/PlainPropsUpgradeTests.cpp
new file mode 100644
--- /dev/null
+++ b/Engine/Plugins/Experimental/PlainProps/Source/Tests/PlainPropsUpgradeTests.cpp
@@ -0,0 +1,770 @@
+// Copyright Epic Games, Inc. All Rights Reserved.
+
+#if WITH_TESTS
+#include "PlainPropsCtti.h"
+#include "PlainPropsIndex.h"
+#include "PlainPropsInternalBind.h"
+#include "PlainPropsInternalBuild.h"
+#include "PlainPropsInternalDiff.h"
+#include "PlainPropsInternalFormat.h"
+#include "PlainPropsInternalRead.h"
+#include "PlainPropsInternalTest.h"
+#include "PlainPropsDiff.h"
+#include "PlainPropsBuild.h"
+#include "PlainPropsBuildSchema.h"
+#include "PlainPropsLoad.h"
+#include "PlainPropsRead.h"
+#include "PlainPropsSave.h"
+#include "PlainPropsUpgrade.h"
+#include "PlainPropsVisualize.h"
+#include "PlainPropsUeCoreBindings.h"
+#include "PlainPropsWrite.h"
+#include "Containers/AnsiString.h"
+#include "Containers/Map.h"
+#include "Templates/UnrealTemplate.h"
+#include "Tests/TestHarnessAdapter.h"
+
+namespace PlainProps::Upgrade::Test
+{
+
+static TIdIndexer<FAnsiString>	GIds;
+static FDebugIds				GDebug(GIds);
+static FEnumDeclarations		GEnums(GDebug);
+static FSchemaBindings			GSchemas(GDebug);
+static FCustomBindingsBottom	GCustoms(GDebug);
+static FCustomBindingsOverlay	GDeltaCustoms(GCustoms);
+static FBindDeclarations		GDeltaDeclarations(GEnums, GDeltaCustoms, GSchemas);
+static FHistory*				GHistory;
+static FScratchAllocator*		GScratch;
+
+struct FFixture
+{
+	FFixture() : History(MakeHistory(GDeltaDeclarations, GIds)), Vis(GIds, "AnsiStr")
+	{
+		GHistory = &*History;
+		GScratch = &Scratch;
+	}
+
+	~FFixture()
+	{
+		GHistory = nullptr;
+		GScratch = nullptr;
+	}
+	
+	FHistoryPtr	History;
+	FScratchAllocator Scratch;
+	DbgVis::FIdScope Vis;
+};
+
+struct FRuntimeIds
+{
+	static FNameId				IndexName(FAnsiStringView Name)		{ return GIds.MakeName(Name); }
+	static FMemberId			IndexMember(FAnsiStringView Name)	{ return GIds.NameMember(Name); }
+	static FConcreteTypenameId	IndexTypename(FAnsiStringView Name)	{ return GIds.NameType(Name); }
+	static FFlatScopeId			IndexScope(FAnsiStringView Name)	{ return GIds.NameScope(Name); }
+	static FEnumId				IndexEnum(FType Type)				{ return GIds.IndexEnum(Type); }
+	static FStructId			IndexStruct(FType Type)				{ return GIds.IndexStruct(Type); }
+	static FIdIndexerBase&		GetIndexer()						{ return GIds; }
+};
+
+struct FDefaultRuntime
+{
+	using Ids = FRuntimeIds;
+	template<class T> using CustomBindings = TCustomBind<T>;
+
+	static FEnumDeclarations&		GetEnums()			{ return GEnums; }
+	static FSchemaBindings&			GetSchemas()		{ return GSchemas; }
+	static FCustomBindings&			GetCustoms()		{ return GCustoms; }
+};
+
+struct FDeltaRuntime : FDefaultRuntime
+{
+	template<class T> using CustomBindings = TCustomDeltaBind<T>;
+
+	static FCustomBindings&			GetCustoms()		{ return GDeltaCustoms; }
+};
+
+////////////////////////////////////////////////////////////////////////////////////////////////
+
+// Temporarily bind obsolete test type to save it
+struct ITempBinding { virtual ~ITempBinding() {} };
+
+template<class T>
+struct TTempBinding : ITempBinding, TScopedStructBinding<T, FDefaultRuntime> {};
+
+using FTempBindingPtr = TUniquePtr<ITempBinding>;
+
+inline constexpr uint32 TestMagics[] = { 0xFEEDF00D, 0xABCD1234, 0xDADADAAA, 0x99887766, 0xF0F1F2F3 };
+
+class FBatchSaver
+{
+public:
+	FBatchSaver(FScratchAllocator& InScratch) : Scratch(InScratch), Types(GEnums, GCustoms, GSchemas) {}
+
+	template<class T>
+	void						Save(T&& Object) 
+	{
+		FDualStructId Id = IndexStructDualId<FRuntimeIds, TTypename<std::remove_reference_t<T>>>();
+		BindObsolete<T>(Id);
+		Roots.Emplace(FDeclId(Id), SaveStruct(&Object, Id, {{GSchemas, GCustoms}, {Scratch}}));
+	}
+
+	void						SetRoots(TArray<FUpgradedInstance>&& In) { Roots = In; }
+	TArray64<uint8>				Write() const;
+
+private:
+	TArray<FUpgradedInstance>	Roots;
+	FScratchAllocator&			Scratch;
+	const FBindDeclarations		Types;
+	TArray<FTempBindingPtr>		Obsoletes;
+
+	template<typename T>
+	void BindObsolete(FDualStructId Id)
+	{
+		if (!Types.Find(Id))
+		{
+			Obsoletes.Emplace(new TTempBinding<T>);
+		}
+	}
+};
+
+TArray64<uint8> FBatchSaver::Write() const
+{
+	// Build partial schemas
+	FSchemasBuilder SchemaBuilders(GIds, Types, Scratch, ESchemaFormat::InMemoryNames);
+	for (FUpgradedInstance Root : Roots)
+	{
+		SchemaBuilders.NoteStructAndMembers(Root.Type, *Root.Value);
+	}
+	FBuiltSchemas Schemas = SchemaBuilders.Build(); 
+
+	// Filter out declared but unused names and ids
+	FWriter Writer(GIds, Types, Schemas, ESchemaFormat::InMemoryNames);
+
+	// Write schemas
+	TArray64<uint8> Out;
+	TArray64<uint8> Tmp;
+	WriteInt(Out, TestMagics[1]);
+	Writer.WriteSchemas(/* Out */ Tmp);
+	WriteAlignmentPadding<uint32>(Out);
+	WriteInt(Out, IntCastChecked<uint32>(Tmp.Num()));
+	WriteArray(Out, Tmp);
+	Tmp.Reset();
+
+	// Write roots
+	WriteInt(Out, TestMagics[2]);
+	for (FUpgradedInstance Root : Roots)
+	{
+		WriteInt(/* out */ Tmp, TestMagics[3]);
+		WriteInt(/* out */ Tmp, Writer.GetWriteId(Root.Type).Get().Idx);
+		Writer.WriteMembers(/* out */ Tmp, Root.Type, *Root.Value);
+		WriteSkippableSlice(Out, Tmp);
+		Tmp.Reset();
+	}
+
+	// Write object terminator
+	WriteSkippableSlice(Out, TConstArrayView64<uint8>());
+	WriteInt(Out, TestMagics[4]);
+		
+	return Out;
+}
+
+////////////////////////////////////////////////////////////////////////////////////////////////
+
+class FBatchLoader
+{
+public:
+	FBatchLoader(FMemoryView Data)
+	{
+		// Read schemas
+		FByteReader It(Data);
+		CHECK(It.Grab<uint32>() == TestMagics[1]);
+		It.SkipAlignmentPadding<uint32>();
+		uint32 SchemasSize = It.Grab<uint32>();
+		FMemoryView SchemasView = It.GrabSlice(SchemasSize);
+		Schemas = ValidateSchemas(SchemasView);
+		FSchemaBatchId Batch = MountReadSchemas(Schemas);
+		CHECK(It.Grab<uint32>() == TestMagics[2]);
+		
+		// Read roots
+		while (uint64 NumBytes = It.GrabVarIntU())
+		{	
+			FByteReader RootIt(It.GrabSlice(NumBytes));
+			CHECK(RootIt.Grab<uint32>() == TestMagics[3]);
+			FStructSchemaId Schema = { RootIt.Grab<uint32>() };
+			Roots.Add({ { Schema, Batch }, RootIt });
+		}
+		
+		CHECK(It.Grab<uint32>() == TestMagics[4]);
+		check(!Roots.IsEmpty());
+
+		// Finally create load plans
+		RuntimeIds = { IndexAllRuntimeIds(*Schemas, GIds) };
+		TConstArrayView<FStructId> StructIds = FInnerIds{RuntimeIds}.GetStructIds(Schemas->NumStructSchemas);
+		Plans = CreateLoadPlans(Batch, GCustoms, GSchemas, StructIds, ESchemaFormat::InMemoryNames);
+	}
+
+	~FBatchLoader()
+	{
+		UnmountReadSchemas(Roots[0].Schema.Batch);
+	}
+	
+	const FSchemaBatch&					GetSchemas() const			{ return *Schemas; }
+	TConstArrayView<FInnerId>			GetRuntimeIds() const		{ return RuntimeIds; }
+	TConstArrayView<FStructView>		GetRoots() const			{ return Roots;	}
+
+	template<class T>
+	T Load(int32 Idx) const
+	{
+		T Out;
+		LoadStruct(&Out, Roots[Idx].Values, Roots[Idx].Schema.Id, *Plans);
+		return MoveTemp(Out);
+	}
+
+private:
+	const FSchemaBatch*			Schemas = nullptr;
+	TArray<FStructView>			Roots;
+	TArray<FInnerId>			RuntimeIds;
+	FLoadBatchPtr				Plans;
+};
+
+////////////////////////////////////////////////////////////////////////////////////////////////
+
+
+static void Run(void (*SaveOld)(FBatchSaver&), void (*TestNew)(FBatchLoader&))
+{
+	// Generate old data
+	TArray64<uint8> OldData;
+	{
+		FScratchAllocator Scratch;
+		FBatchSaver OldBatch(Scratch);
+		SaveOld(OldBatch);
+		OldData = OldBatch.Write();
+	}
+
+	TArray64<uint8> NewData;
+	{
+		// Mount old schemas
+		FBatchLoader OldBatch(MakeMemoryView(OldData));
+
+		// Upgrade old instances
+		FBatchUpgrader Upgrader(OldBatch.GetSchemas(), OldBatch.GetRuntimeIds(), GIds);
+		FBindDeclarations Declarations(GEnums, GCustoms, GSchemas);
+		CHECK(!Upgrader.MatchSchemas(Declarations));
+		FScratchAllocator Scratch;
+		TArray<FUpgradedInstance> NewRoots = Upgrader.Upgrade(OldBatch.GetRoots(), *GHistory, Scratch);
+
+		// Write new schemas and instances
+		FBatchSaver NewBatch(Scratch);
+		NewBatch.SetRoots(MoveTemp(NewRoots));
+		NewData = NewBatch.Write();
+	}
+	
+	// Test new data
+	FBatchLoader NewBatch(MakeMemoryView(NewData));
+	TestNew(NewBatch);
+}
+
+////////////////////////////////////////////////////////////////////////////////////////////////
+
+// Tests that everything was read
+struct FTestMemberReader : public FMemberReader
+{
+	using FMemberReader::FMemberReader;
+	
+	~FTestMemberReader()
+	{
+		CHECK(MemberIdx == NumMembers); // Must read all members
+		CHECK(RangeTypeIdx == NumRangeTypes); // Must read all ranges
+#if DO_CHECK
+		CHECK(InnerSchemaIdx == NumInnerSchemas); // Must read all schema ids
+#endif
+	}
+};
+
+template<typename OutType, typename InType>
+TArray<OutType> MakeArray(const InType& Items)
+{
+	TArray<OutType> Out;
+	Out.Reserve(IntCastChecked<int32>(Items.Num()));
+	for (const auto& Item : Items)
+	{
+		Out.Emplace(Item);
+	}
+	return Out;
+}
+
+////////////////////////////////////////////////////////////////////////////////////////////////
+////////////////////////////////////////////////////////////////////////////////////////////////
+////////////////////////////////////////////////////////////////////////////////////////////////
+
+// No ordering impact - only helps compile out old upgraders
+enum ERelease { R0, R1, R2, R3 };
+
+enum class ERename0 : uint8 {Same, OldA, OldB1, OldB2};
+enum class ERename : uint8 {Same, NewA, NewB};
+enum class EFlagify : uint8 {A, B};
+enum class EFlatten : uint8 {None, A, B, AB};
+PP_REFLECT_ENUM(PlainProps::Upgrade::Test, ERename0, Same, OldA, OldB1, OldB2);
+PP_REFLECT_ENUM(PlainProps::Upgrade::Test, ERename, Same, NewA, NewB);
+PP_REFLECT_ENUM(PlainProps::Upgrade::Test, EFlagify, A, B);
+PP_REFLECT_ENUM(PlainProps::Upgrade::Test, EFlatten, A, B, AB);
+//
+//template<> PlainProps::Upgrade::Chronicle<ERename>(FHistory& Out)
+//{
+//	FEnumChronicle Rename(Out, "", "ERename");
+//
+//	//Rename<R0>.Rename("OldA", "NewA");
+//	//Rename<R0>.Rename("OldB1", "NewB");
+//	//Rename<R0>.Rename("OldB2", "NewB");
+//}
+//
+//template<> PlainProps::Upgrade::Chronicle<EFlagify>(FHistory& Out)
+//{
+//	FEnumChronicle Flagify(Out, "", "EFlagify");
+//	Flagify.Record<R0>([](FEnumChronicle& Out)
+//		{
+//			//Out.
+//		});
+//	//Flagify<R0>.Flagify("None", {});
+//	//Flagify<R0>.Flagify("AB", {"A", "B"});
+//}
+//
+//template<> PlainProps::Upgrade::Chronicle<EFlatten>(FHistory& Out)
+//{
+//	FEnumChronicle Flatten(Out, "", "EFlatten");
+//	//Flatten<R0>.Flatten({}, "None");
+//	//Flatten<R0>.Flatten({"A", "B"}, "AB");
+//}
+
+////////////////////////////////////////////////////////////////////////////////////////////////
+
+enum class EFlag : uint8 {A = 8, B = 32};
+PP_REFLECT_ENUM(PlainProps::Upgrade::Test, EFlag, A, B);
+
+struct FPoint { float X, Y; };
+PP_REFLECT_STRUCT(PlainProps::Upgrade::Test, FPoint, void, X, Y); 
+inline bool operator==(FPoint A, FPoint B) { return A.X == B.X && A.Y == B.Y; }
+
+struct FBasicOld
+{
+	uint32		Widen; // Could be handled automatically by loosen + fasten
+	int32		Rename;
+	uint16		Reorder; // Handled automatically by loosen + fasten
+	double		Drop;
+	double		Unchanged;
+	bool		FlagA;
+	bool		FlagB;
+	float		PointX;
+	float		PointY;
+	FGuid		Guid; // Renamed to GuidRenamed to test struct renaming
+};
+PP_REFLECT_STRUCT(PlainProps::Upgrade::Test, FBasicOld, void, Widen, Rename, Reorder, Drop, Unchanged, FlagA, FlagB, PointX, PointY, Guid); 
+
+struct FBasic
+{
+	uint64		Widen;
+	uint16		Reorder;
+	int32		Renamed;
+	double		Unchanged;
+	EFlag		Flags;
+	FPoint		Point;
+	uint64		GuidAB; // Todo: Switch to uint64 Guid[2] once static array support is built out
+	uint64		GuidCD;
+
+	bool operator==(const FBasic&) const = default;
+};
+PP_REFLECT_STRUCT(PlainProps::Upgrade::Test, FBasic, void, Widen, Reorder, Renamed, Unchanged, Flags, Point, GuidAB, GuidCD);
+
+
+// Sugar-free API
+void Record(TChroniclesOf<FBasic> Out)
+{
+	FChronicler BasicOld(Out, "PlainProps::Upgrade::Test", "FBasicOld");
+
+	BasicOld.In<R1>([](IOps& Out)
+	{	
+		Out.Add({"Widen"_U32}, {"Widen"_U64}, [](FIns In, FOuts Out, FContext Ctx)
+			{
+				Out[0].U64() = In[0].U32();
+			});
+		
+		Out.Add({"Rename"_S32}, {"Renamed"_S32}, [](FIns In, FOuts Out, FContext Ctx) 
+			{ 
+				Out[0].S32() = In[0].S32();
+			});
+
+		Out.Add({"Drop"_F64}, {}, [](FIns, FOuts, FContext) {});
+		
+		// Partial declaration of the members / enumerators that will be accessed
+		// Note a chronicle might need multiple declarations for a type if it changed during chronicle
+		FEnumFlagNames AB = Out.NameFlags({"PlainProps::Upgrade::Test", "EFlag"}, {"A", "B"});
+
+		// Executed last if only FlagA or FlagB was sparsely saved
+		Out.Add({~"FlagA"_B1, ~"FlagB"_B1}, {"Flags"_Enum(AB)}, [](FIns In, FOuts Out, FContext)
+			{
+				FOutEnumFlags& Flags = Out[0].EnumFlags();
+				Flags[/* A */ 0] = In[/* FlagA */ 0].OptBool();
+				Flags[/* B */ 1] = In[/* FlagB */ 1].OptBool();
+			});
+	});
+	
+	FChronicler Basic(Out, "PlainProps::Upgrade::Test", "FBasic");
+	BasicOld.Rename<R1>(Basic.GetName());
+	Basic.In<R1>([](IOps& Out)
+	{
+		// Partial declaration of the members / enumerators that will be accessed
+		// Note a chronicle might need multiple declarations for a type if it changed during chronicle
+		
+		FStructMemberNames XY = Out.NameMembers({"PlainProps::Upgrade::Test", "FPoint"}, {"X"_F32, "Y"_F32});
+		Out.Add({ ~"PointX"_F32, ~"PointY"_F32 }, { "Point"_Struct(XY) }, [](FIns In, FOuts Out, FContext)
+			{
+				TOuts<2> Point = Out[0].StructMembers<2>();
+				Point[/* X */ 0].F32() = In[/* PointX */ 0].OptF32().Get(0.f);
+				Point[/* Y */ 1].F32() = In[/* PointY */ 1].OptF32().Get(0.f);
+			});
+		
+		FStructId GuidId = Out.NameStruct({/*"/Script/CoreUObject"*/ "", "Guid"});
+		Out.Add({"Guid"_Struct(GuidId)}, {"GuidRenamed"_Struct(GuidId)}, [](FIns In, FOuts Out, FContext Ctx) 
+			{ 
+				Out[0].Struct() = In[0].Struct();
+			});
+
+		FStructMemberNames ABCD = Out.NameMembers(GuidId, {"A"_U32, "B"_U32, "C"_U32, "D"_U32});
+		Out.Add({ "GuidRenamed"_Struct(ABCD) }, { "GuidAB"_U64, "GuidCD"_U64 }, [](FIns In, FOuts Out, FContext)
+			{
+				TIns<4> ABCD = In[0].StructMembers<4>();
+
+				uint64 A = ABCD[0].U32();
+				uint64 B = ABCD[1].U32();
+				uint64 C = ABCD[2].U32();
+				uint64 D = ABCD[3].U32();
+				
+				Out[/* GuidAB */ 0].U64() = (A << 32) | B;
+				Out[/* GuidCD */ 1].U64() = (C << 32) | D;
+			});
+	});
+}
+
+// Sugar / template sorcery (up to a few members) for type-safety & usability
+//void Record(TChroniclesOf<FBasic> Out)
+//{
+//	FChronicler BasicOld(Out, "PlainProps::Upgrade::Test", "FBasicOld");	
+//	BasicOld.In<R1>([](IOps& Out)
+//	{
+//		Out.Convert<uint32>("Widen"_U64);
+//		Out.Rename("Rename"_S32, "Renamed");
+//		Out.Drop("Drop"_F64);
+//
+//		FEnumFlagNames Flag = Out.NameFlags({"PlainProps::Upgrade::Test", "EFlag"}, {"A", "B"});
+//		Out.Enumify({~"FlagA"_B1, ~"FlagB"_B1}, "Flags"_Enum(Flag));
+//	});
+//	
+//	FChronicler Basic(Out, "PlainProps::Upgrade::Test", "FBasic");
+//	BasicOld.Rename<R1>(Basic.GetName());
+//	Basic.In<R1>([](IOps& Out)
+//	{
+//		FStructMemberNames XY = Out.NameMembers({"PlainProps::Upgrade::Test", "FPoint"}, {"X"_F32, "Y"_F32});
+//		Out.Structify({ ~"PointX"_F32, ~"PointY"_F32 }, "Point"_Struct(XY));
+//		
+//		FStructId GuidId = Out.NameStruct({/*"/Script/CoreUObject"*/ "", "Guid"});
+//		Out.Rename("Guid"_Struct(GuidId), "GuidRenamed");
+//
+//		FStructMemberNames ABCD = Out.NameMembers(GuidId, {"A"_U32, "B"_U32, "C"_U32, "D"_U32});
+//		Out.Destructify("GuidRenamed"_Struct(ABCD), { "GuidAB"_U64, "GuidCD"_U64 }, 
+//			[](uint32 A, uint32 B, uint32 C, uint32 D, TOut<2, uint64> Out, FContext Ctx)
+//			{
+//				Out[/* GuidAB */ 0].U64() = (A << 32) | B;
+//				Out[/* GuidCD */ 1].U64() = (C << 32) | D;
+//			});
+//	});
+//}
+
+
+// Sugar-free but with type-checked inputs
+//
+// Can type-check inputs and outputs too if we specify number of input params, i.e. AddTypechecked<4>
+//
+//void Record(TChroniclesOf<FBasic> Out)
+//{
+//	FChronicler BasicOld(Out, "PlainProps::Upgrade::Test", "FBasicOld");
+//
+//	BasicOld.In<R1>([](IOps& Out)
+//	{	
+//		Out.AddTypechecked({"Widen"_U32}, {"Widen"_U64}, [](uint32 Widen, FOuts Out, FContext Ctx)
+//			{
+//				Out[0].U64() = Widen;
+//			});
+//		
+//		Out.AddTypechecked({"Rename"_S32}, {"Renamed"_S32}, [](int32 Rename, FOuts Out, FContext Ctx) 
+//			{ 
+//				Out[0].S32() = Rename;
+//			});
+//
+//		Out.AddTypechecked({"Drop"_F64}, {}, [](double Drop, FOuts, FContext) {});
+//		
+//		FEnumFlagNames AB = Out.NameFlags({"PlainProps::Upgrade::Test", "EFlag"}, {"A", "B"});
+//		Out.AddTypechecked({~"FlagA"_B1, ~"FlagB"_B1}, {"Flags"_Enum(AB)}, 
+//			[](TOptional<bool> FlagA, TOptional<bool> FlagB, FOuts Out, FContext)
+//			{
+//				FOutEnumFlags& Flags = Out[0].EnumFlags();
+//				Flags[/* A */ 0] = FlagA;
+//				Flags[/* B */ 1] = FlagB;
+//			});
+//	});
+//	
+//	FChronicler Basic(Out, "PlainProps::Upgrade::Test", "FBasic");
+//	BasicOld.Rename<R1>(Basic.GetName());
+//	Basic.In<R1>([](IOps& Out)
+//	{
+//		// Partial declaration of the members / enumerators that will be accessed
+//		// Note a chronicle might need multiple declarations for a type if it changed during chronicle
+//		
+//		FStructMemberNames XY = Out.NameMembers({"PlainProps::Upgrade::Test", "FPoint"}, {"X"_F32, "Y"_F32});
+//		Out.AddTypechecked({ ~"PointX"_F32, ~"PointY"_F32 }, { "Point"_Struct(XY) },
+//			[](TOptional<float> PointX, TOptional<float> PointY, FOuts Out, FContext)
+//			{
+//				TOuts<2> Point = Out[0].StructMembers<2>();
+//				Point[/* X */ 0].F32() = PointX.Get(0.f);
+//				Point[/* Y */ 1].F32() = PointY.Get(0.f);
+//			});
+//		
+//		FStructId GuidId = Out.NameStruct({/*"/Script/CoreUObject"*/ "", "Guid"});
+//		Out.AddTypechecked({"Guid"_Struct(GuidId)}, {"GuidRenamed"_Struct(GuidId)}, [](FInStruct Guid, FOuts Out, FContext Ctx) 
+//			{ 
+//				Out[0].Struct() = Guid;
+//			});
+//
+//		FStructMemberNames ABCD = Out.NameMembers(GuidId, {"A"_U32, "B"_U32, "C"_U32, "D"_U32});
+//		Out.AddTypechecked({ "GuidRenamed"_Struct(ABCD) }, { "GuidAB"_U64, "GuidCD"_U64 }, [](TIns<4> ABCD, FOuts Out, FContext)
+//			{
+//				uint64 A = ABCD[0].U32();
+//				uint64 B = ABCD[1].U32();
+//				uint64 C = ABCD[2].U32();
+//				uint64 D = ABCD[3].U32();
+//				
+//				Out[/* GuidAB */ 0].U64() = (A << 32) | B;
+//				Out[/* GuidCD */ 1].U64() = (C << 32) | D;
+//			});
+//	});
+//}
+
+////////////////////////////////////////////////////////////////////////////////////////////////
+
+struct FStr
+{
+	TArray<char8_t> Data;
+};
+
+struct FPoints
+{
+	TArray<float> Xs;
+	TArray<float> Ys;
+};
+
+// Test convert FRangesA into FRangesB and vice versa
+struct FRangesA
+{
+	TArray<uint8>			Leaves;	// Int <-> Enum
+	TArray<char8_t>			String; // Structify <-> Destructify
+	TArray<FPoint>			Points; // SoA <-> AoS
+	TArray<TArray<FPoint>>	Series; // Join <-> Split
+	TArray<char>			Unterminated; // Add <-> Remove null terminator
+};
+PP_REFLECT_STRUCT(PlainProps::Upgrade::Test, FRangesA, void, Leaves, String, Points, Series, Unterminated); 
+
+struct FRangesB
+{
+	TArray<EFlag>			Leaves;
+	FStr					String;
+	FPoints					Points;
+	TArray<FPoint>			Series;
+	TArray<int32>			SeriesLengths;
+	TArray<char>			Terminated;
+};
+PP_REFLECT_STRUCT(PlainProps::Upgrade::Test, FRangesB, void, Leaves, String, Points, Series, SeriesLengths, Terminated);
+
+//
+//// Sugar
+//template<> PlainProps::Upgrade::Chronicle<FRangesA>(FHistory& Out)
+//{
+//	TEnum<2, R0> Flag(Out, "", "EFlag", {"A", "B"});
+//	TStruct<2, R0> Point(Out, "", "FPoint", {"X", "Y"});
+//	TStruct<2, R0> Points(Out, "", "FPoints", {"Xs", "Ys"});
+//	TStruct<2, R0> Str(Out, "", "FStr", {"Data"});
+//
+//	FChronicle RangesA(Out, "", "FRangesA"); // ~ registers chronicle once we've collected all member names
+//	RangesA.Transform("Leaves"_U8, "Leaves"_Enum(Flag)),
+//		[](uint8 In, FOutEnum Out)
+//		{
+//			Out.EmitIf(/* A */ 0, !!(In & 1));
+//			Out.EmitIf(/* B */ 1, !!(In & 2));
+//		});
+//	RangesA.Structify("String"_UTF8, "String"_Struct(Str));  // overload without lambda
+//	RangesA.Structify("Points"_Range.Of(Point), "Points"_Struct(Points),
+//		[](FInStructs Points, FOutFloats Xs, FOutFloats Ys)
+//		{
+//			FFloatWriter XIt = Xs.Emit(In.Num());
+//			FFloatWriter YIt = Ys.Emit(In.Num());
+//			for (FInStruct Point : Points)
+//			{
+//				XIt.Add(Point[/* X */ 0].AsRequired<float>());
+//				YIt.Add(Point[/* Y */ 1].AsRequired<float>());
+//			}
+//		});
+//
+//	RangesA.Upgrade("Series"_Range.Nest().Of(Point), "Series"_Range.Of(Point), "SeriesLengths"_Range.Of<int32>(),
+//		[](FInNestedStructs SeriesOfSeries, FOutStructs Series, FOutI32s SeriesLengths)
+//		{
+//			FStructWriter SeriesIt = Series.Emit(SumNums(SeriesOfSeries));
+//			FI32Writer LengthsIt = SeriesLengths.Emit(SeriesOfSeries.Num());
+//			for (FInStructRange Points : SeriesOfSeries)
+//			{
+//				LengthsIt.Add(IntCastChecked<int32>(Points.Num()));
+//				for (FInStruct Point : Points)
+//				{
+//					SeriesIt.Add(Point);
+//				}
+//			}
+//		});
+//
+//	RangesA.Upgrade<FInUtf8s>("Unterminated", "Terminated"),
+//		[](FInUtf8s Unterminated, FOutUtf8s Terminated)
+//		{
+//			FUtf8s OutData = Terminated.Emit(Unterminated.Num() + 1);
+//			OutData.EmitSlice(0, Unterminated);
+//			OutData.EmitItem(Unterminated.Num, '\0');
+//		});
+//}
+//
+//// Sugar (move down)
+//template<> PlainProps::Upgrade::Chronicle<FRangesB>(FUpgrades& Out)
+//{
+//	TEnum<2, R0> Flag(Out, "", "EFlag", {"A", "B"});
+//	TStruct<2, R0> Point(Out, "", "FPoint", {"X", "Y"});
+//	TStruct<2, R0> Points(Out, "", "FPoints", {"Xs", "Ys"});
+//	TStruct<2, R0> Str(Out, "", "FStr", {"Data"});
+//
+//	FChronicle RangesB(Out, "", "FRangesB"); // ~ registers chronicle once we've collected all member names
+//	
+//	RangesB.Transform<FInEnum>(MakeIn(Flag, "Leaves"), "Leaves"),
+//		[](FInEnum In, FOutU8 Out)
+//		{
+//			OutU8.Emit(uint8(In[/* A */ 0]) + uint8(In[/* B */ 1])*2);
+//		});
+//	RangesB.Destructify<FInUtf8s>(MakeInStruct("String"), "String"); // overload without lambda
+//	
+//	RangesB.Destructify<FInFloats, FInFloats>(MakeInStruct(Points, "Points"), MakeOutStructs(Point, "Points")),
+//		[](FInFloats Xs, FInFloats Ys, FOutNestedStructs Points)
+//		{
+//			uint64 Num = Xs.Num();
+//			if (Num != Ys.Num())
+//			{
+//				Ctx.LogDataLoss();
+//				Num = FMath::Min(Xs.Num(), Ys.Num());
+//			}
+//
+//			FStructsWriter Out = Points.Emit(Xs.Num());
+//			for (uint64 Idx = 0; Idx < Num; ++Idx)
+//			{
+//				Out[Idx][0].Emit(Xs[Idx]);
+//				Out[Idx][1].Emit(Ys[Idx]);
+//			}
+//		});
+//	
+//	RangesB.Upgrade<FInStructs, FI32s>(MakeIn(Point, "Series"), "SeriesLengths", MakeOutStructs(Point, "Series"),
+//		[](FInStructs Series, FInI32s SeriesLengths, FOutStructs Out)
+//		{
+//			FStructWriter* OutIt = Out.Emit(SeriesOfSeries.Num()).GetData();
+//			uint64 SeriesIdx = 0;
+//			for (uint64 Len : SeriesOfSeries)
+//			{
+//				FOutStruct OutPoint = *OutIt++;
+//				OutPoint[0].Emit(Series.Slice(SeriesIdx, Len));
+//				SeriesIdx += Len;
+//			}
+//		});
+//
+//	RangesB.Upgrade<FInUtf8s>("Terminated", "Unterminated"),
+//		[](FInUtf8s Terminated, FOutUtf8s Unterminated)
+//		{
+//			check(Terminated.Num());
+//			check(Terminated.Last() == '\0');
+//			FUtf8s OutData = Unterminated.Emit(Terminated.Num() - 1);
+//			OutData.EmitSlice(0, Terminated, Terminated.Num() - 1);
+//		});
+//}
+
+template<class T>
+struct TScopedStructUpgradeBinding 
+{
+	TScopedStructUpgradeBinding()
+	{
+		Record(static_cast<TChroniclesOf<T>>(FChronicles(/* out */ *GHistory, Binding.DeclId)));
+	}
+	
+	~TScopedStructUpgradeBinding()
+	{
+		EraseChroniclesOf(*GHistory, Binding.DeclId);
+	}
+
+	TScopedStructBinding<T, FDefaultRuntime> Binding;
+};
+
+////////////////////////////////////////////////////////////////////////////////////////////////
+
+TEST_CASE_NAMED(FPlainPropsUpgradeTest, "System::PlainProps::Upgrade", "[PlainProps][SmokeFilter]")
+{
+	FFixture Fixture;
+
+	SECTION("Enums")
+	{
+	}
+
+	SECTION("Basics")
+	{
+		TScopedStructUpgradeBinding<FBasic> Basic;
+		TScopedEnumDeclaration<EFlag, EEnumMode::Flag, FDefaultRuntime> Flag;
+		TScopedStructBinding<FPoint, FDefaultRuntime> Point;
+
+		Run([](FBatchSaver& Batch)
+		{
+			Batch.Save(FBasicOld{
+				.Widen=1, .Rename=2, .Reorder=3, .Drop=4, .Unchanged=5, .FlagA=false, .FlagB=true, 
+				.PointX=6.f, .PointY=7.f, .Guid=FGuid(0x01234567, 0x89ABCDEF,0,0xFFFFFFFF)});
+		}, 
+		[](FBatchLoader& Batch)
+		{
+			CHECK(Batch.Load<FBasic>(0) == FBasic{
+				.Widen=1, .Reorder=3, .Renamed=2, .Unchanged=5, .Flags=EFlag::B, 
+				.Point={6.f, 7.f}, .GuidAB=0x0123456789ABCDEF, .GuidCD=0xFFFFFFFF});
+		});
+	}
+
+	SECTION("Ranges")
+	{
+	}
+	
+	SECTION("Versions")
+	{
+	}
+
+	SECTION("Inheritance")
+	{
+	}
+
+	SECTION("Unknowns")
+	{
+	}
+
+	SECTION("Errors")
+	{
+	}
+
+	SECTION("Defaults")
+	{
+	}
+}
+
+} // namespace PlainProps::Upgrade::Test
+#endif // WITH_TESTS
\ No newline at end of file
