Fixed clang compilationi error for GetTypeHash

#rnx
#Synced-CL 48906160

diff --git a/Engine/Plugins/Experimental/PlainProps/Source/Public/PlainPropsDiff.h b/Engine/Plugins/Experimental/PlainProps/Source/Public/PlainPropsDiff.h
--- a/Engine/Plugins/Experimental/PlainProps/Source/Public/PlainPropsDiff.h
+++ b/Engine/Plugins/Experimental/PlainProps/Source/Public/PlainPropsDiff.h
@@ -92,7 +92,7 @@ bool PLAINPROPS_API DiffSchemas(FSchemaBatchId A, FSchemaBatchId B);
 // Nodes in the path are stored in forward order. Sibling links are invalidated.
 struct FDiffPath : public TArray<FDiffNode> 
 {
-	[[nodiscard]] PLAINPROPS_API friend uint32 GetTypeHash(const FDiffPath& Path);
+	PLAINPROPS_API friend uint32 GetTypeHash(const FDiffPath& Path);
 };
 
 PLAINPROPS_API bool operator==(const FDiffPath& A, const FDiffPath& B);
