PlainProps first simple mismatch implementation + some misc fixes

* Add EMismatch describing mismatch origin
* Fuse mismatches with values in loose representation
* Provide both loose and fixed value in FUpgradedInstance
* Stop fastening instances with any mismatches
* Fix freeing scratch addresses via FOpsPtr
* Add static assert to catch excessive stack usage in static TMap/Set bindings

#rnx
#Synced-CL 48909003

diff --git a/Engine/Plugins/Experimental/PlainProps/Source/Private/PlainPropsInternalUpgrade.h b/Engine/Plugins/Experimental/PlainProps/Source/Private/PlainPropsInternalUpgrade.h
--- a/Engine/Plugins/Experimental/PlainProps/Source/Private/PlainPropsInternalUpgrade.h
+++ b/Engine/Plugins/Experimental/PlainProps/Source/Private/PlainPropsInternalUpgrade.h
@@ -21,15 +21,6 @@ struct FVersionUpgrades;
 
 ////////////////////////////////////////////////////////////////////////////////////////////////
 
-// Per struct-indexed member id - allows TArray<T> lookup instead of TMap<FMemberId, T>
-struct FParameterId
-{
-	uint16 Idx = 0xFFFF;
-
-	bool operator==(FParameterId O) const { return Idx == O.Idx; }
-	friend uint32 GetTypeHash(FParameterId Id) { return Id.Idx; };
-};
-
 struct FEnumeratorIndexer : private TSet<FNameId>
 {
 	FEnumeratorIndexer() {}
@@ -115,6 +106,7 @@ struct FStructMember
 	FParameterId			Name;
 	uint16					Version;
 	ELooseState				State;
+	EMismatch				Mismatch;
 	uint64					Num;	// Num loose members or FByteReader::CheckableBytes (0 in shipping config)
 	FStructData				Data;
 };
@@ -267,7 +259,7 @@ struct FVersionUpgrades
 	FMiniSlice					NameMatches[0];		// #NumNames
 	
 	template<class T>
-	const T&					At(uint32 Offset) const		{ return *reinterpret_cast<const T*>(reinterpret_cast<const uint8*>(this) + OutputsOffset); }
+	const T&					At(uint32 Offset) const		{ return *reinterpret_cast<const T*>(reinterpret_cast<const uint8*>(this) + Offset); }
 
 	const FMemberMatch*			GetMatches() const			{ return AlignPtr<FMemberMatch>(NameMatches + NumNames); }
 	const FOp*					GetOps() const				{ return AlignPtr<FOp>(GetMatches() + NumMatches); }
@@ -450,19 +442,17 @@ class FInstanceUpgrader
 	FLooseValue						Decompose(FLooseType Type, FLooseValue Value, const FEnumDecomposition& Decomp);
 	FLooseValue						Decompose(FEnumMember Input, const FEnumDecomposition& Decomp);
 	FLooseValue						Decompose(FStructMember Input, const FStructDecomposition& Decomp);
-	TArrayView<FLooseMember>		Decompose(FStructView Input, const FStructDecomposition& Decomp, FParameterId Name);
+	TArrayView<FLooseMember>		Decompose(FStructView Input, const FStructDecomposition& Decomp, FParameterId Name, EMismatch& OutMismatch);
 	FEnumIndexData					Reorder(FEnumeratorId Flat, const TSet<FEnumeratorId>& Order);
 	FEnumIndexData					Reorder(FEnumeratorIdSet Flags, const TSet<FEnumeratorId>& Order);
-	TArrayView<FLooseMember>		Reorder(TConstArrayView<FLooseMember> Unsorted, TConstArrayView<FLooseParameter> Params, const TSet<FParameterId>& Order);
 	void							Input(FMemberMatch Match, FLooseValue Value);
 	void							Stash(FLooseType Type, FLooseValue Value);
-	bool							Emit(FLooseType Type, FLooseValue Value);
 	void							Run(FOp Op, const FIn* Inputs, uint8 Recipient);
 	FOuter							Resolve(uint8 Recipient);
 	void							Recompose(TArrayView<FOut> Outputs, TConstArrayView<FOutParameter> Expected);
 	FEnumeratorData					Recompose(FInnermostEnumType Type, FEnumIndexData In, TConstArrayView<FEnumeratorId> Order);
 	void							Flush();
-	FLooseStruct					Finish();
+	FLooseStruct					Finish(EMismatch& OutMismatch);
 
 	FString							Print(FParameterId Parameter) const;
 	template<class T> FString		Print(T Item) const { return Ctx.Debug.Print(Item); }
@@ -470,8 +460,8 @@ class FInstanceUpgrader
 public:
 	FInstanceUpgrader(FBatchContext& Ctx, const FStructUpgrades& Upgrades, uint16 Version);
 	
-	FLooseStruct					Upgrade(FStructView Source, const FLooseSchema& From, const FLooseSchema* To);
-	FLooseStruct					Upgrade(FLooseStruct In, const FLooseSchema* To);
+	FLooseStruct					Upgrade(FStructView Source, const FLooseSchema& From, const FLooseSchema* To, EMismatch& OutMismatch);
+	FLooseStruct					Upgrade(FLooseStruct In, const FLooseSchema* To, EMismatch& OutMismatch);
 	FStructMember					Upgrade(FStructMember In, const FLooseSchema* To, FOuter InOuter = {});
 
 	void							Reuse(); // Todo: Optimization for upgrading range-of-structs
@@ -487,7 +477,18 @@ struct FSchemaCache
 
 ////////////////////////////////////////////////////////////////////////////////////////////////
 
-using FOpsPtr = TUniquePtr<FOps>;
+template <typename T>
+struct TDestroyOnly
+{
+	void operator()(T* Ptr) const { Ptr->~T(); } 
+};
+
+template<class T>
+using TScratchPtr = TUniquePtr<T, TDestroyOnly<T>>;
+
+////////////////////////////////////////////////////////////////////////////////////////////////
+
+using FOpsPtr = TScratchPtr<FOps>;
 using FStructUpgradesPtr = TUniquePtr<FStructUpgrades>;
 
 struct FHistory
diff --git a/Engine/Plugins/Experimental/PlainProps/Source/Private/PlainPropsUpgrade.cpp b/Engine/Plugins/Experimental/PlainProps/Source/Private/PlainPropsUpgrade.cpp
--- a/Engine/Plugins/Experimental/PlainProps/Source/Private/PlainPropsUpgrade.cpp
+++ b/Engine/Plugins/Experimental/PlainProps/Source/Private/PlainPropsUpgrade.cpp
@@ -155,14 +155,9 @@ static void SetLooseFlag(FEnumeratorIdSet& Out, FEnumeratorId Name, bool bValue)
 
 ////////////////////////////////////////////////////////////////////////////////////////////////
 
-inline FParameterId GetName(FLooseMetadata Meta)
-{
-	return { static_cast<uint16>(Meta.Name) };
-}
-
 inline FParameterId GetName(const FLooseMember& Member)
 {
-	return { static_cast<uint16>(Member.GetValue().Meta.Name) };
+	return Member.GetValue().Meta.Name;
 }
 
 inline ELooseState GetInnermostState(FLooseMetadata Meta)
@@ -184,7 +179,7 @@ inline FStructMember AsStruct(FLooseType Type, FLooseMetadata Meta, FStructData
 	check(Type.IsStruct());
 	FDeclId Id = Type.InnermostId.Get().AsStructDeclId();
 	FOptionalStructSchemaId Schema = ToOptionalStruct(Meta.InnermostSchema);
-	return { Id, Schema, GetName(Meta), Type.InnermostVersion, GetInnermostState(Meta), Meta.Num, Data };
+	return { Id, Schema, Meta.Name, Type.InnermostVersion, GetInnermostState(Meta), Meta.Mismatch, Meta.Num, Data };
 }
 
 inline FEnumMember AsEnum(FLooseType Type, FLooseMetadata Meta, FEnumData Data)
@@ -192,7 +187,7 @@ inline FEnumMember AsEnum(FLooseType Type, FLooseMetadata Meta, FEnumData Data)
 	check(Type.IsEnum());
 	FEnumId Id = Type.InnermostId.Get().AsEnum();
 	FOptionalEnumSchemaId Schema = ToOptionalEnum(Meta.InnermostSchema);
-	return { Id, Schema, GetName(Meta), Type.InnermostVersion, GetInnermostState(Meta), Type.InnermostType.Enum.FlagMode, IntCastChecked<uint16>(Meta.Num), Data };
+	return { Id, Schema, Meta.Name, Type.InnermostVersion, GetInnermostState(Meta), Type.InnermostType.Enum.FlagMode, IntCastChecked<uint16>(Meta.Num), Data };
 }
 
 inline FRangeMember AsRange(FLooseMember In)
@@ -225,7 +220,7 @@ inline FStructView AsSource(FStructMember In, FSchemaBatchId Batch)
 static FLooseMember AsMember(FStructMember In)
 {
 	FLooseType Type = {.InnermostType = InnermostStaticStruct, .InnermostVersion = In.Version, .InnermostId = FInnerId(In.Id) };
-	FLooseMetadata Meta = { .IsSet = 1, .IsFullyUpgraded = In.State == ELooseState::UpgradedLoose, .Name = In.Name.Idx, .Num = In.Num };
+	FLooseMetadata Meta = { .IsSet = 1, .IsFullyUpgraded = In.State == ELooseState::UpgradedLoose, .Name = In.Name, .Num = In.Num };
 	FLooseData Data = { .Struct = In.Data };	
 	return FLooseMember(Type, {Meta, Data});
 }
@@ -261,13 +256,13 @@ inline const FLooseMember& ReadOutput(const FOut& Out) { return Out; }
 inline TArrayView<FLooseMember> EditOutputs(TArrayView<FOut> Outs) { return MakeArrayView(static_cast<FLooseMember*>(Outs.GetData()), Outs.Num()); }
 
 FOut::FOut(FOutParameter Param)
-: FLooseMember(Param.Type, FLooseValue{{.Name = Param.Name.Idx, .Num = Param.NumInners}, {}})
+: FLooseMember(Param.Type, FLooseValue{{.Name = Param.Name, .Num = Param.NumInners}, {}})
 {
 	check(Param.NumInners == 0 || Param.Type.IsEnum());
 }
 
 FOut::FOut(FOutParameter Param, FOut* Members)
-: FLooseMember(Param.Type, FLooseValue{{.IsSet = 1, .Name = Param.Name.Idx, .Num = Param.NumInners}, {.Struct = {.Loose = Members}}})
+: FLooseMember(Param.Type, FLooseValue{{.IsSet = 1, .Name = Param.Name, .Num = Param.NumInners}, {.Struct = {.Loose = Members}}})
 {
 	check(Param.NumInners > 0);
 }
@@ -438,13 +433,13 @@ FLooseValue FLooseMemberReader::Grab()
 	else if (Next.Type.IsStruct())
 	{
 		FStructView Struct = FMemberReader::GrabStruct();
-		FLooseMetadata Meta = { .IsSet = 1, .Name = Next.Name.Idx, .InnermostSchema = Struct.Schema.Id, .Num = Struct.Values.CheckableSize() };
+		FLooseMetadata Meta = { .IsSet = 1, .Name = Next.Name, .InnermostSchema = Struct.Schema.Id, .Num = Struct.Values.CheckableSize() };
 		return { Meta, {.Struct = { Struct.Values.Peek() }} };
 	}
 	
 	FLeafView Leaf = GrabLeaf();
 	uint64 Value = GetIntegerValue(Leaf);
-	FLooseMetadata Meta = { .IsSet = 1, .Name = Next.Name.Idx };
+	FLooseMetadata Meta = { .IsSet = 1, .Name = Next.Name };
 
 	if (Leaf.Enum)
 	{
@@ -471,7 +466,7 @@ FLooseValue FLooseMemberReader::Grab()
 FRangeMember FLooseMemberReader::ToLooseRange(FParameterId Name, FLooseType Type, FRangeView Range)
 {
 	FOptionalSchemaId SchemaId = GetSchema(Range).InnermostSchema;
-	FLooseMetadata Meta = { .IsSet = 1, .IsSourceRange = 1, .Name = Name.Idx, .InnermostSchema = SchemaId, .Num = Range.Num() };
+	FLooseMetadata Meta = { .IsSet = 1, .IsSourceRange = 1, .Name = Name, .InnermostSchema = SchemaId, .Num = Range.Num() };
 	if (Type.NumRanges > 1)
 	{
 		Meta.IsSourceRange = 0;
@@ -658,7 +653,7 @@ FInstanceUpgrader::FInstanceUpgrader(FBatchContext& InCtx, const FStructUpgrades
 
 static FLooseMember MakeMember(FLooseParameter Param)
 {
-	FLooseValue UnsetNamedValue = {{.Name = Param.Name.Idx}, {}};
+	FLooseValue UnsetNamedValue = {{.Name = Param.Name}, {}};
 	return FLooseMember(Param.Type, UnsetNamedValue);
 }
 
@@ -765,58 +760,61 @@ FORCENOINLINE static FLooseMember UpgradeUnmatchedMember(FLooseType Type, FLoose
 
 void FInstanceUpgrader::Give(FLooseType Type, FLooseValue Value, FMatchHelper Helper)
 {
-	FParameterId Name = GetName(Value.Meta);
-	if (Name.Idx >= NumNames) // Unknown member, potentially from obsolete type
+	FParameterId Name = Value.Meta.Name;
+	if (Name.Idx < NumNames)
 	{
-		Stash(Type, Value);
-		return;
-	}
-
-	const FMiniSlice NameMatches = Helper.NameMatches[Name.Idx];
-	uint8& ScanIdx = NameToScanIndices[Name.Idx];
+		const FMiniSlice NameMatches = Helper.NameMatches[Name.Idx];
+		uint8& ScanIdx = NameToScanIndices[Name.Idx];
 Rematch:
-	for (FMemberMatch Match : Helper.AllMatches.Slice(NameMatches.Idx + ScanIdx, NameMatches.Num - ScanIdx))
-	{
-		++ScanIdx;
-		check(ScanIdx);
-		if (Match.Type.Compare(Type)) // todo: check code-gen and change WithoutMetadata impl if needed
+		for (FMemberMatch Match : Helper.AllMatches.Slice(NameMatches.Idx + ScanIdx, NameMatches.Num - ScanIdx))
 		{
-			if (Match.Decomposes())
+			++ScanIdx;
+			check(ScanIdx);
+			if (Match.Type.Compare(Type)) // todo: check code-gen and change WithoutMetadata impl if needed
 			{
-				Value = Decompose(Type, Value, Match.Idx.Decomp);
+				if (Match.Decomposes())
+				{
+					Value = Decompose(Type, Value, Match.Idx.Decomp);
+					if (Value.IsMismatch())
+					{
+						checkSlow(Value.Meta.Mismatch == EMismatch::Inners || Value.Meta.Mismatch == EMismatch::Inputs);
+						Value.Meta.Mismatch = EMismatch::Inputs;
+						goto Done;
+					}
+
+					checkf(ScanIdx < NameMatches.Num, TEXT("Input decomposition match can't be last match for '%s' in '%s'"), *Print(Name), *Print(Struct.Id));
+					Match = Helper.AllMatches[NameMatches.Idx + ScanIdx];
+					++ScanIdx;
+					checkf(Match.Type.Compare(Type), TEXT("The match after an input decomposition match should be of same type"));
+					checkf(!Match.Decomposes(), TEXT("Can't input decompose twice"));
+				}
 
-				checkf(ScanIdx < NameMatches.Num, TEXT("Input decomposition match can't be last match for '%s' in '%s'"), *Print(Name), *Print(Struct.Id));
-				Match = Helper.AllMatches[NameMatches.Idx + ScanIdx];
-				++ScanIdx;
-				checkf(Match.Type.Compare(Type), TEXT("The match after an input decomposition match should be of same type"));
-				checkf(!Match.Decomposes(), TEXT("Can't input decompose twice"));
+				Input(Match, Value);
+				return;
 			}
-
-			Input(Match, Value);
-			return;
 		}
-	}
 
-	// Unmatched
+		// Unmatched
 
-	if (const FTransformation* Xform = FindUpgrade(Type, Value.Meta, Ctx))
-	{
-		// Todo: Only upgrade to next matching version, e.g. upgrading V1 -> V2 -> V3 would miss a V2 match.
-		//		 Determine next version by looking at matches
-		FLooseMember Upgraded = UpgradeUnmatchedMember(Type, Value, *Xform, Ctx);
-		checkSlow(Name.Idx == GetName(Upgraded).Idx);
+		if (const FTransformation* Xform = FindUpgrade(Type, Value.Meta, Ctx))
+		{
+			// Todo: Only upgrade to next matching version, e.g. upgrading V1 -> V2 -> V3 would miss a V2 match.
+			//		 Determine next version by looking at matches
+			FLooseMember Upgraded = UpgradeUnmatchedMember(Type, Value, *Xform, Ctx);
+			checkSlow(Name.Idx == GetName(Upgraded).Idx);
 
-		bool bVersionUpdated = !Upgraded.GetType().Compare(Type);
-		Type = Upgraded.GetType();
-		Value = Upgraded.GetValue();
+			bool bVersionUpdated = !Upgraded.GetType().Compare(Type);
+			Type = Upgraded.GetType();
+			Value = Upgraded.GetValue();
 
-		if (bVersionUpdated)
-		{
-			ScanIdx = 0;
-			goto Rematch;
+			if (bVersionUpdated)
+			{
+				ScanIdx = 0;
+				goto Rematch;
+			}
 		}
 	}
-
+Done:
 	Stash(Type, Value);
 }
 
@@ -828,39 +826,42 @@ void FInstanceUpgrader::Give(FLooseType Type, FLooseValue Value)
 
 void FInstanceUpgrader::Stash(FLooseType Type, FLooseValue Value)
 {
-	if (!Emit(Type, Value))
-	{
-		Mismatches.Emplace(Type, Value);
-	}
-}
-
-bool FInstanceUpgrader::Emit(FLooseType Type, FLooseValue Value)
-{
-	FParameterId Name = GetName(Value.Meta);
-	if (Name.Idx >= static_cast<uint32>(Outcome.Num()))
+	EMismatch Mismatch = Value.Meta.Mismatch;
+	FParameterId Name = Value.Meta.Name;
+	if (Mismatch != EMismatch::No)
+	{}
+	else if (Name.Idx >= static_cast<uint32>(Outcome.Num()))
 	{
-		return false;
+		Mismatch = EMismatch::Surplus;
 	}
 	else if (FLooseMember& Out = Outcome[Name.Idx]; Out)
 	{
-		return false;
+		Mismatch = EMismatch::Duplicate;
+		Out.EditTypechecked().SetFirstMismatch(EMismatch::Duplicate);
 	}
 	else if (bUntypedOutcome)
 	{
 		Out = FLooseMember(Type, Value);
-		return true;
+		return;
 	}
 	else if (Type.Compare(Out.GetType()))
 	{
 		Out.SetTypechecked(Value);
-		return true;
+		return;
 	}
-	return false;
+	else
+	{
+		Mismatch = EMismatch::Type;
+	}
+
+	Value.Meta.Mismatch = Mismatch;
+	Mismatches.Emplace(Type, Value);
 }
 
 void FInstanceUpgrader::Input(FMemberMatch Match, FLooseValue Value)
 {
 	check(Match.Idx.Op < Current->NumOps);
+	check(Value.Meta.Mismatch == EMismatch::No);
 	FOp Op =  Current->GetOps()[Match.Idx.Op];
 	FIn In(Match.Type, Value);
 	if (Op.NumInputs() == 1)
@@ -898,64 +899,6 @@ void FInstanceUpgrader::Input(FMemberMatch Match, FLooseValue Value)
 	}
 }
 
-static TArrayView<FLooseMember> NewMembersInInputOrder(TConstArrayView<FLooseParameter> Params, FBatchContext& Ctx)
-{
-	int32 Num = Params.Num();
-	FLooseMember* OutIt = static_cast<FLooseMember*>(Ctx.Scratch.Allocate(Num * sizeof(FLooseMember), alignof(FLooseMember)));
-	for (FLooseParameter Param : Params)
-	{
-		new (OutIt++) FLooseMember(MakeMember(Param));
-	}
-	return MakeArrayView(OutIt - Num, Num);
-}
-
-TArrayView<FLooseMember> FInstanceUpgrader::Reorder(TConstArrayView<FLooseMember> Unsorted, TConstArrayView<FLooseParameter> Params, const TSet<FParameterId>& Order)
-{
-	check(Params.Num() == Order.Num());
-
-	TArrayView<FLooseMember> Sorted = NewMembersInInputOrder(Params, Ctx);
-	for (FLooseMember Src : Unsorted)
-	{
-		if (Src)
-		{
-			FParameterId Name = GetName(Src);
-			FSetElementId Place = Order.FindId(Name);
-			check(Place.IsValidId()); // todo: emit Mismatches
-			FLooseMember& Dst = Sorted[Place.AsInteger()];
-			check(Src.GetType().Compare(Dst.GetType()));
-			Dst.SetTypechecked(Src.GetValue());
-		}
-	}
-
-	return Sorted;
-}
-
-void SetValueOrAddMismatch(FLooseMember* Out, FLooseType Type, FLooseValue Value, TArray<FLooseMember>& OutMismatches)
-{
-	check(Value.Meta.IsSet);
-	if (Out && Type.Compare(Out->GetType()))
-	{
-		Out->SetTypechecked(Value);
-	}
-	else
-	{
-		OutMismatches.Emplace(Type, Value);
-	}
-}
-
-static TArrayView<FLooseMember> Prune(TArrayView<FLooseMember> In, int32 Num, TArray<FLooseMember>& Mismatches)
-{
-	for (FLooseMember Surplus : In.RightChop(Num))
-	{
-		if (Surplus.GetValue())
-		{
-			Mismatches.Add(Surplus);
-		}
-	}
-
-	return In.Left(Num);
-}
-
 FLooseValue FInstanceUpgrader::Decompose(FLooseType Type, FLooseValue Value, uint16 DecompIdx)
 {
 	check(DecompIdx < Current->NumDecompositions);
@@ -989,7 +932,7 @@ FLooseValue FInstanceUpgrader::Decompose(FEnumMember In, const FEnumDecompositio
 	FEnumIndexData Indices = In.FlagMode	? Reorder(Enumerators.Flags, Decomp.InputOrder) 
 											: Reorder(Enumerators.Flat, Decomp.InputOrder);
 
-	FLooseMetadata Meta = { .IsSet = 1, .Name = In.Name.Idx };
+	FLooseMetadata Meta = { .IsSet = 1, .Name = In.Name };
 	return { Meta, { .Enum = { .Loose = { .Indices = Indices } } } };
 }
 
@@ -1017,15 +960,109 @@ FEnumIndexData FInstanceUpgrader::Reorder(FEnumeratorIdSet Flags, const TSet<FEn
 	return { .Flags { .Bits = Indices } };
 }
 
+static TArrayView<FLooseMember> Fuse(TArrayView<FLooseMember> As, TConstArrayView<FLooseMember> Bs, FScratchAllocator& Scratch)
+{
+	// Compact As
+	int32 Pos = 0;
+	for (FLooseMember& A : As)
+	{
+		if (A || A.GetValue().Meta.Mismatch == EMismatch::Missing)
+		{
+			As[Pos++] = A;
+		}
+	}
+
+	// Place mismatches last, might contain EMismatch::Duplicate
+	Algo::StableSort(As.Left(Pos), [](FLooseMember A1, FLooseMember A2) { return A1.GetValue().IsMismatch() < A2.GetValue().IsMismatch(); });
+
+	// Place Bs after compacted As, reallocate if needed
+	TArrayView<FLooseMember> Out(As.GetData(), Pos + Bs.Num());
+	if (Out.Num() > As.Num())
+	{
+		static_assert(std::is_trivially_copyable_v<FLooseMember> && std::is_trivially_destructible_v<FLooseMember>);
+		Out = MakeArrayView(reinterpret_cast<FLooseMember*>(Scratch.Allocate(Out.NumBytes(), alignof(FLooseMember))), Out.Num());
+		FMemory::Memcpy(Out.GetData(), As.GetData(), As.Left(Pos).NumBytes());
+	}
+	FMemory::Memcpy(&Out[Pos], Bs.GetData(), Bs.NumBytes());
+
+	return Out;
+}
+
+static void AddInput(FLooseMember* Expected, FLooseType Type, FLooseValue Value, TArray<FLooseMember>& OutUnexpected)
+{
+	check(Value.Meta.IsSet);
+	if (Expected && Type.Compare(Expected->GetType()))
+	{
+		Expected->SetTypechecked(Value);
+	}
+	else
+	{
+		Value.Meta.Mismatch = Expected ? EMismatch::Type : EMismatch::Surplus;
+		OutUnexpected.Emplace(Type, Value);
+	}
+}
+
+static TArrayView<FLooseMember> FinalizeInputs(TArrayView<FLooseMember> Expected, TConstArrayView<FLooseMember> Unexpected, FScratchAllocator& Scratch, EMismatch& OutMismatch)
+{
+	// Flag missing inputs as mismatches
+	for (FLooseMember& Input : Expected)
+	{
+		if (!Input && !Input.GetType().OptionalParameter)
+		{
+			OutMismatch = EMismatch::Inputs;
+			Input.EditTypechecked().Meta.Mismatch = EMismatch::Missing;
+		}
+	}
+
+	// Fuse surplus and type mismatches
+	if (!Unexpected.IsEmpty())
+	{
+		OutMismatch = EMismatch::Inputs;
+		return Fuse(Expected, Unexpected, Scratch);
+	}
+
+	return Expected;
+}
+
+static TArrayView<FLooseMember> NewMembersInInputOrder(TConstArrayView<FLooseParameter> Params, FScratchAllocator& Scratch)
+{
+	int32 Num = Params.Num();
+	FLooseMember* OutIt = static_cast<FLooseMember*>(Scratch.Allocate(Num * sizeof(FLooseMember), alignof(FLooseMember)));
+	for (FLooseParameter Param : Params)
+	{
+		new (OutIt++) FLooseMember(MakeMember(Param));
+	}
+	return MakeArrayView(OutIt - Num, Num);
+}
+
+static TArrayView<FLooseMember> ReorderInputs(TConstArrayView<FLooseMember> Unsorted, TConstArrayView<FLooseParameter> Params, const TSet<FParameterId>& Order, FScratchAllocator& Scratch, EMismatch& OutMismatch)
+{
+	check(Params.Num() == Order.Num());
+
+	TArray<FLooseMember> Unexpected;
+	TArrayView<FLooseMember> Expected = NewMembersInInputOrder(Params, Scratch);
+	for (FLooseMember Src : Unsorted)
+	{
+		if (Src)
+		{
+			FSetElementId Place = Order.FindId(GetName(Src));
+			AddInput(Place.IsValidId() ? &Expected[Place.AsInteger()] : nullptr, Src.GetType(), Src.GetValue(), /* out */ Unexpected);
+		}
+	}
+
+	return FinalizeInputs(Expected, Unexpected, Scratch, OutMismatch);
+}
+
 FLooseValue FInstanceUpgrader::Decompose(FStructMember In, const FStructDecomposition& Decomp)
 {
 	check(In.State != ELooseState::Unset);
 
 	const FLooseSchema& To = *Decomp.InputSchema;
 	TArrayView<FLooseMember> Members;
+	EMismatch Mismatch = EMismatch::No;
 	if (In.State == ELooseState::Source)
 	{
-		Members = Decompose(AsSource(In, Ctx.Batch), Decomp, In.Name);
+		Members = Decompose(AsSource(In, Ctx.Batch), Decomp, In.Name, /* out */ Mismatch);
 	}
 	else
 	{
@@ -1034,6 +1071,7 @@ FLooseValue FInstanceUpgrader::Decompose(FStructMember In, const FStructDecompos
 			FInstanceUpgrader Upgrader(Ctx, *Decomp.Upgrades, To.Version);
 			FStructMember Upgraded = Upgrader.Upgrade(In, &To);
 			Members = MakeArrayView(Upgraded.Data.Loose, static_cast<int32>(Upgraded.Num));
+			Mismatch = Upgraded.Mismatch;
 		}
 		else
 		{
@@ -1042,19 +1080,23 @@ FLooseValue FInstanceUpgrader::Decompose(FStructMember In, const FStructDecompos
 	
 		if (!Decomp.InputOrder.IsEmpty())
 		{
-			Members = Reorder(Members, To.GetMembers(), Decomp.InputOrder);
+			Members = ReorderInputs(Members, To.GetMembers(), Decomp.InputOrder, Ctx.Scratch, /* out */ Mismatch);
 		}
 		else if (Members.Num() > To.NumMembers) // Declaration order
 		{
-			Members = Prune(Members, To.NumMembers, /* out */ Mismatches);
+			Mismatch = EMismatch::Inputs;
+			for (FLooseMember& Surplus : Members.RightChop(To.NumMembers))
+			{
+				Surplus.EditTypechecked().Meta.Mismatch = EMismatch::Surplus;
+			}
 		}
 	}
-	check(Members.Num() == To.NumMembers);
+	check(Members.Num() == To.NumMembers || Mismatch != EMismatch::No);
 	
-	return { {.IsSet = 1, .Name = In.Name.Idx, .Num = static_cast<uint32>(Members.Num())}, {.Struct = {.Loose = Members.GetData()}} };
+	return { {.IsSet = 1, .Mismatch = Mismatch, .Name = In.Name, .Num = static_cast<uint32>(Members.Num())}, {.Struct = {.Loose = Members.GetData()}} };
 }
 
-TArrayView<FLooseMember> FInstanceUpgrader::Decompose(FStructView In, const FStructDecomposition& Decomp, FParameterId Name)
+TArrayView<FLooseMember> FInstanceUpgrader::Decompose(FStructView In, const FStructDecomposition& Decomp, FParameterId Name, EMismatch& OutMismatch)
 {
 	const FTransformation& Xform = Ctx.Transformations[In.Schema.Id];
 	const FLooseSchema& From = *Xform.SourceSchema;
@@ -1063,15 +1105,16 @@ TArrayView<FLooseMember> FInstanceUpgrader::Decompose(FStructView In, const FStr
 	if (Decomp.Upgrades)
 	{
 		FInstanceUpgrader Upgrader(Ctx, *Decomp.Upgrades, To.Version);
-		FLooseStruct Upgraded = Upgrader.Upgrade(In, From, &To);
+		FLooseStruct Upgraded = Upgrader.Upgrade(In, From, &To, OutMismatch);
+		OutMismatch = OutMismatch == EMismatch::Inners ? EMismatch::Inputs : OutMismatch;
 		TArrayView<FLooseMember> Members(Upgraded.Members, Upgraded.Num);
-		return Decomp.InputOrder.IsEmpty() ? Members : Reorder(Members, To.GetMembers(), Decomp.InputOrder);
+		return Decomp.InputOrder.IsEmpty() ? Members : ReorderInputs(Members, To.GetMembers(), Decomp.InputOrder, Ctx.Scratch, OutMismatch);
 	}
 
 	// Loosen source into input order
 	check(!Xform.Upgrades);
-
-	TArrayView<FLooseMember> Members = NewMembersInInputOrder(To.GetMembers(), Ctx);
+	TArrayView<FLooseMember> Expected = NewMembersInInputOrder(To.GetMembers(), Ctx.Scratch);
+	TArray<FLooseMember> Unexpected; // Mismatches
 	FLooseMemberReader It(To, In, ViewLooseners(Ctx), Ctx.Scratch);
 	const TSet<FParameterId>& Order = Decomp.InputOrder;
 	if (Order.IsEmpty())
@@ -1079,8 +1122,7 @@ TArrayView<FLooseMember> FInstanceUpgrader::Decompose(FStructView In, const FStr
 		while (It.HasMore())
 		{
 			FLooseParameter Param = It.Peek();
-			FLooseValue Value = It.Grab();
-			SetValueOrAddMismatch(Param.Name.Idx < Members.Num() ? &Members[Param.Name.Idx] : nullptr, Param.Type, Value, /* out */ Mismatches);
+			AddInput(Param.Name.Idx < Expected.Num() ? &Expected[Param.Name.Idx] : nullptr, Param.Type, It.Grab(), /* out */ Unexpected);
 		}	
 	}
 	else
@@ -1088,13 +1130,12 @@ TArrayView<FLooseMember> FInstanceUpgrader::Decompose(FStructView In, const FStr
 		while (It.HasMore())
 		{
 			FLooseParameter Param = It.Peek();
-			FLooseValue Value = It.Grab();
 			FSetElementId Place = Order.FindId(Param.Name);
-			SetValueOrAddMismatch(Place.IsValidId() ? &Members[Place.AsInteger()] : nullptr, Param.Type, Value, /* out */ Mismatches);
+			AddInput(Place.IsValidId() ? &Expected[Place.AsInteger()] : nullptr, Param.Type, It.Grab(), /* out */ Unexpected);
 		}
 	}
 
-	return Members;
+	return FinalizeInputs(Expected, Unexpected, Ctx.Scratch, OutMismatch);
 }
 
 FOuter FInstanceUpgrader::Resolve(uint8 Recipient)
@@ -1227,13 +1268,20 @@ static void RenameMembers(TArrayView<FLooseMember> Members, const FParameterInde
 {
 	for (FLooseMember& Member : Members)
 	{
-		FOptionalMemberId Name = From.Resolve(GetName(Member));
-		FParameterId Reindexed = To.Index(Name);
-		Member.Rename(Reindexed.Idx);
+		if (Member || Member.IsMismatch())
+		{
+			FOptionalMemberId Name = From.Resolve(GetName(Member));
+			FParameterId Reindexed = To.Index(Name);
+			Member.Rename(Reindexed);
+		}
+		else
+		{
+			Member.Rename({});
+		}
 	}
 }
 
-FLooseStruct FInstanceUpgrader::Finish()
+FLooseStruct FInstanceUpgrader::Finish(EMismatch& OutMismatch)
 {
 	Flush();
 
@@ -1242,6 +1290,21 @@ FLooseStruct FInstanceUpgrader::Finish()
 		RenameMembers(Outcome, Ctx.LooseIds.Get(Struct.Id), Ctx.LooseIds.Get(Struct.FinalId));
 	}
 
+	for (const FLooseMember& Member : Outcome)
+	{
+		if (!Member && !Member.IsOptional() &&
+			!Algo::FindByPredicate(Mismatches, [&Member](const FLooseMember& Mismatch) { return Member.GetValue().Meta.Name == Mismatch.GetValue().Meta.Name; }))
+		{
+			Mismatches.Add_GetRef(Member).EditTypechecked().Meta.Mismatch = EMismatch::Missing;
+		}
+	}
+
+	if (Mismatches.Num() > 0)
+	{
+		OutMismatch = EMismatch::Inners;
+		Outcome = Fuse(Outcome, Mismatches, Ctx.Scratch);
+	}
+
 	return { Struct.FinalId, Struct.FinalVersion, static_cast<uint16>(Outcome.Num()), Outcome.GetData() };
 }
 
@@ -1250,25 +1313,26 @@ FString FInstanceUpgrader::Print(FParameterId Parameter) const
 	return Ctx.Debug.Print(Ctx.LooseIds.Get(Struct.Id).Resolve(Parameter));
 }
 
-FLooseStruct FInstanceUpgrader::Upgrade(FStructView In, const FLooseSchema& From, const FLooseSchema* To)
+FLooseStruct FInstanceUpgrader::Upgrade(FStructView In, const FLooseSchema& From, const FLooseSchema* To, EMismatch& OutMismatch)
 {
 	check(Current == GetVersion(Struct, From.Version));
 	Expect(To);
 	Feed(FLooseMemberReader(From, In, ViewLooseners(Ctx), Ctx.Scratch));
-	return Finish();
+	return Finish(OutMismatch);
 }
 
-FLooseStruct FInstanceUpgrader::Upgrade(FLooseStruct In, const FLooseSchema* To)
+FLooseStruct FInstanceUpgrader::Upgrade(FLooseStruct In, const FLooseSchema* To, EMismatch& OutMismatch)
 {
 	check(Current == GetVersion(Struct, In.Version));
 	Expect(To);
 	Feed(MakeArrayView(In.Members, In.Num));
-	return Finish();
+	return Finish(OutMismatch);
 }
 
 FStructMember FInstanceUpgrader::Upgrade(FStructMember In, const FLooseSchema* To, FOuter InOuter)
 {
 	check(Current == GetVersion(Struct, In.Version));
+	check(In.Mismatch == EMismatch::No);
 
 	Outer = InOuter;
 	Expect(To);
@@ -1293,12 +1357,12 @@ FStructMember FInstanceUpgrader::Upgrade(FStructMember In, const FLooseSchema* T
 			break;
 	}
 	
-	FLooseStruct Out = Finish();
+	EMismatch Mismatch = EMismatch::No;
+	FLooseStruct Out = Finish(Mismatch);
 	// todo: upgrade limit should return ELooseState::Loose
-	return { Out.Id, NoId, In.Name, Out.Version, ELooseState::UpgradedLoose, Out.Num, { .Loose = Out.Members } };
+	return { Out.Id, NoId, In.Name, Out.Version, ELooseState::UpgradedLoose, Mismatch, Out.Num, { .Loose = Out.Members } };
 }
 
-
 ////////////////////////////////////////////////////////////////////////////////////////////////
 
 struct FImplicitFlags
@@ -1932,12 +1996,14 @@ void FBatchRun::CreateTransformations(const TBitArray<>& MatchingSchemas, const
 static FLooseType IndexObsolete(FLiteralIndexerBase& Ids)
 {
 	FStructId Obsolete = Ids.IndexStruct(FType(NoId, FTypenameId(Ids.IndexTypenameLiteral("__obsolete__"))));
-	return { .InnermostType = InnermostStaticStruct, .InnermostVersion = 0xBAAD, .InnermostId = FInnerId(Obsolete) };
+	return { .OptionalParameter = 1, .InnermostType = InnermostStaticStruct, .InnermostVersion = 0xBAAD, .InnermostId = FInnerId(Obsolete) };
 }
 
+// @param In isn't mismatched
 static FLooseStruct ReorderAsDeclaration(FLooseStruct In, const FLooseSchema& Declared, FScratchAllocator& Scratch)
 {
 	// opt: check if In members have consecutive names, if so type-check and return In
+	//		could also compact and only return IsSet members and reuse the existing In allocation
 
 	FLooseMember* OutIt = static_cast<FLooseMember*>(Scratch.AllocateZeroed(Declared.NumMembers * sizeof(FLooseMember), alignof(FLooseMember)));
 	for (FLooseParameter Param : Declared.GetMembers())
@@ -1951,7 +2017,7 @@ static FLooseStruct ReorderAsDeclaration(FLooseStruct In, const FLooseSchema& De
 		if (Src)
 		{
 			FParameterId Name = GetName(Src);
-			check(Name.Idx < Out.Num()); // todo: emit Mismatches
+			check(Name.Idx < Out.Num()); // Should've been a mismatch
 			FLooseMember& Dst = Out[Name.Idx];
 			check(!Dst);
 			check(Src.GetType().Compare(Dst.GetType()));
@@ -1966,16 +2032,17 @@ FUpgradedInstance FBatchRun::UpgradeStruct(FStructView Source, FTransformation X
 {
 	FBatchContext Ctx = { Transformations, Scratch, History.Ids, History.LooseIds, IndexObsolete(History.Ids), Source.Schema.Batch, {Scratch}, {}, Debug };
 	FInstanceUpgrader SourceUpgrader(Ctx, *Xform.Upgrades, Xform.SourceSchema->Version);
-	FLooseStruct Upgraded = SourceUpgrader.Upgrade(Source, *Xform.SourceSchema, Xform.TargetSchema);
+	EMismatch Mismatch = EMismatch::No;
+	FLooseStruct Upgraded = SourceUpgrader.Upgrade(Source, *Xform.SourceSchema, Xform.TargetSchema, /* out */ Mismatch);
 
 	// Handle renamed structs
-	while (Upgraded.Id != Xform.Upgrades->Id)
+	while (Upgraded.Id != Xform.Upgrades->Id && Mismatch == EMismatch::No)
 	{
 		Xform = Transformations[Upgraded.Id];
 		if (Xform.Upgrades)
 		{
 			FInstanceUpgrader Upgrader(Ctx, *Xform.Upgrades, Upgraded.Version);
-			Upgraded = Upgrader.Upgrade(Upgraded, Xform.TargetSchema);
+			Upgraded = Upgrader.Upgrade(Upgraded, Xform.TargetSchema, /* out */ Mismatch);
 		}
 		else
 		{
@@ -1983,7 +2050,8 @@ FUpgradedInstance FBatchRun::UpgradeStruct(FStructView Source, FTransformation X
 		}
 	}
 
-	return { Upgraded.Id, FastenStruct(Upgraded) };
+	const FBuiltStruct* Fixed = Mismatch == EMismatch::No ? FastenStruct(Upgraded) : nullptr;
+	return { Upgraded.Id, MakeArrayView(Upgraded.Members, Upgraded.Num), Fixed };
 }
 
 FUpgradedInstance FBatchRun::UpgradeOrRebuild(FStructView Source)
@@ -1992,7 +2060,7 @@ FUpgradedInstance FBatchRun::UpgradeOrRebuild(FStructView Source)
 	FStructSchemaId SchemaId = Source.Schema.Id;
 	FDeclId Id = RuntimeIds[SchemaId];
 	FTransformation Xform = Transformations[Source.Schema.Id];
-	return Xform.Upgrades ? UpgradeStruct(Source, Xform) : FUpgradedInstance{ Id, RebuildStruct(Source, Id) };
+	return Xform.Upgrades ? UpgradeStruct(Source, Xform) : FUpgradedInstance{ Id, {}, RebuildStruct(Source, Id) };
 }
 
 static FUnpackedLeafType ToLeafType(EInnermostKind Arithmetic)
@@ -2117,6 +2185,8 @@ FBuiltValue FBatchRun::FastenMember(FLooseMember Member) const
 
 FBuiltStruct* FBatchRun::FastenStruct(FStructMember Struct) const
 {
+	check(Struct.Mismatch == EMismatch::No);
+
 	if (Struct.State == ELooseState::Source)
 	{
 		return RebuildStruct(AsSource(Struct, SourceBatch));
@@ -3117,10 +3187,11 @@ void FHistory::Link(TConstArrayView<FOpsPtr> Committed)
 }
 
 template<class T>
-static const T& GetDecomposition(const FVersionUpgrades& Version, uint16 DecompIdx)
+T& EditDecomposition(const FVersionUpgrades& Version, uint16 DecompIdx)
 {
 	check(DecompIdx < Version.NumDecompositions);
-	return Version.At<T>(Version.GetDecompOffsets()[DecompIdx]);
+	uint32 Offset = Version.GetDecompOffsets()[DecompIdx];
+	return const_cast<T&>(Version.At<T>(Offset));
 }
 
 void FHistory::Link(FDeclId Compiled)
@@ -3134,7 +3205,7 @@ void FHistory::Link(FDeclId Compiled)
 			{
 				if (Match.Type.IsStruct())
 				{
-					Link(const_cast<FStructDecomposition&>(Version->At<FStructDecomposition>(Match.Idx.Decomp)));
+					Link(EditDecomposition<FStructDecomposition>(*Version, Match.Idx.Decomp));
 				}
 			}
 		}
diff --git a/Engine/Plugins/Experimental/PlainProps/Source/Public/PlainPropsDeclare.h b/Engine/Plugins/Experimental/PlainProps/Source/Public/PlainPropsDeclare.h
--- a/Engine/Plugins/Experimental/PlainProps/Source/Public/PlainPropsDeclare.h
+++ b/Engine/Plugins/Experimental/PlainProps/Source/Public/PlainPropsDeclare.h
@@ -117,6 +117,7 @@ PLAINPROPS_API FStructDeclarationPtr Declare(FStructSpec Spec);
 
 struct IDeclarations
 {
+	virtual ~IDeclarations() {}
 	virtual FDeclId								Lower(FBindId BindId) const = 0;
 	virtual const FEnumDeclaration*				Find(FEnumId Id) const = 0;
 	virtual const FStructDeclaration*			Find(FStructId Id) const = 0;
diff --git a/Engine/Plugins/Experimental/PlainProps/Source/Public/PlainPropsLooseTypes.h b/Engine/Plugins/Experimental/PlainProps/Source/Public/PlainPropsLooseTypes.h
--- a/Engine/Plugins/Experimental/PlainProps/Source/Public/PlainPropsLooseTypes.h
+++ b/Engine/Plugins/Experimental/PlainProps/Source/Public/PlainPropsLooseTypes.h
@@ -114,6 +114,15 @@ inline constexpr FLooseType ReflectLoose = { 0, 0, ReflectInnermost<T>, 0, NoId
 
 //////////////////////////////////////////////////////////////////////////////////////////////////
 
+// Per struct-indexed member id - allows TArray<T> lookup instead of TMap<FMemberId, T>
+struct FParameterId
+{
+	uint16 Idx = 0xFFFF;
+
+	bool operator==(FParameterId O) const { return Idx == O.Idx; }
+	friend uint32 GetTypeHash(FParameterId Id) { return Id.Idx; };
+};
+
 // Per enum-indexed enumerator name
 struct FEnumeratorId
 {
@@ -197,15 +206,18 @@ union FLooseData
 	FStructData			Struct;
 };
 
+enum class EMismatch : uint8 { No, Inners, Inputs, Type, Duplicate, Surplus, Missing };
+
 // FLooseData metadata
 struct FLooseMetadata
 {
 	uint8				IsSet = 0;
 	uint8				IsSourceRange : 1 = 0;
 	uint8				IsFullyUpgraded : 1 = 0;
-	uint16				Name = 0;					// FParameterId / FEnumeratorId
-	FOptionalSchemaId	InnermostSchema;			// Set if innermost is source struct/enum
-	uint64				Num = 0;					// Range items, loose struct members or inner members / enumerators during op execution
+	EMismatch			Mismatch : 6 = EMismatch::No;	// Only need 4 bits currently but want all bits initialized
+	FParameterId		Name;
+	FOptionalSchemaId	InnermostSchema;				// Set if innermost is source struct/enum
+	uint64				Num = 0;						// Range items, loose struct members or inner members / enumerators during op execution
 };
 
 struct FLooseValue
@@ -213,14 +225,25 @@ struct FLooseValue
 	FLooseMetadata		Meta;
 	FLooseData			Data;
 
-	explicit operator	bool() const { return !!Meta.IsSet; }
+	explicit operator	bool() const			{ return !!Meta.IsSet; }
+	bool				IsMismatch() const		{ return Meta.Mismatch != EMismatch::No; }
+	void				SetNamed(FLooseValue In);				// Assign but keep name
+	void				SetFirstMismatch(EMismatch Mismatch);
+};
 
-	void				SetNamed(FLooseValue In)	// Assign but keep name
+inline void FLooseValue::SetNamed(FLooseValue In)
+{
+	In.Meta.Name = Meta.Name;
+	*this = In;
+}
+
+inline void FLooseValue::SetFirstMismatch(EMismatch Mismatch)
+{
+	if (Meta.Mismatch == EMismatch::No)
 	{
-		In.Meta.Name = Meta.Name;
-		*this = In;
+		Meta.Mismatch = Mismatch;
 	}
-};
+}
 
 //////////////////////////////////////////////////////////////////////////////////////////////////
 
@@ -231,12 +254,13 @@ public:
 	FLooseMember(FLooseType T, FLooseValue V) : Type(T), Value(V) {}
 	
 	explicit			operator bool() const			{ return !!Value.Meta.IsSet; }
+	bool				IsMismatch() const				{ return Value.IsMismatch(); }
+	bool				IsOptional() const				{ return Type.OptionalParameter; }
 	FLooseType			GetType() const					{ return Type; }
 	FLooseValue			GetValue() const				{ return Value; }
-	void				Rename(uint16 ParameterId)		{ Value.Meta.Name = ParameterId; }
+	void				Rename(FParameterId Name)		{ Value.Meta.Name = Name; }
 	void				SetTypechecked(FLooseValue V)	{ Value = V; } 
 	FLooseValue&		EditTypechecked()				{ return Value; } 
-	bool				IsOptional() const				{ return Type.OptionalParameter; }
 
 protected:
 	FLooseType			Type;
diff --git a/Engine/Plugins/Experimental/PlainProps/Source/Public/PlainPropsUeCoreBindings.h b/Engine/Plugins/Experimental/PlainProps/Source/Public/PlainPropsUeCoreBindings.h
--- a/Engine/Plugins/Experimental/PlainProps/Source/Public/PlainPropsUeCoreBindings.h
+++ b/Engine/Plugins/Experimental/PlainProps/Source/Public/PlainPropsUeCoreBindings.h
@@ -922,6 +922,7 @@ struct TBaseDeltaBinding : ICustomBinding, FBaseDeltaBinding
 			}
 			else
 			{
+				static_assert(sizeof(T) <= 128, "Switch to e.g. TArray<uint8, TInlineAllocator<32>> + SetNumInitialized(sizeof(T) + over-alignment) + AlignPtr(GetData(), alignof(T))");
 				alignas(T) uint8 Buffer[sizeof(T)];
 				ConstructAndLoadStruct(Buffer, Item);
 				T& Tmp = *reinterpret_cast<T*>(Buffer);
diff --git a/Engine/Plugins/Experimental/PlainProps/Source/Public/PlainPropsUpgrade.h b/Engine/Plugins/Experimental/PlainProps/Source/Public/PlainPropsUpgrade.h
--- a/Engine/Plugins/Experimental/PlainProps/Source/Public/PlainPropsUpgrade.h
+++ b/Engine/Plugins/Experimental/PlainProps/Source/Public/PlainPropsUpgrade.h
@@ -624,10 +624,15 @@ PLAINPROPS_API FHistoryPtr MakeHistory(IDeclarations& Latest, FLiteralIndexerBas
 
 ////////////////////////////////////////////////////////////////////////////////////////////////
 
+PLAINPROPS_API FMemberId			ResolveName(const FLooseMember& Member, const FHistory& History);
+PLAINPROPS_API FNameId				ResolveEnumerator(const FLooseMember& FlatEnum, const FHistory& History);
+PLAINPROPS_API TArray<FNameId>		ResolveFlags(const FLooseMember& FlagEnum, const FHistory& History);
+
 struct FUpgradedInstance
 {
-	FDeclId					Type;
-	const FBuiltStruct*		Value;
+	FDeclId										Type;
+	TConstArrayView<FLooseMember>				Loose; // Owned by scratch, empty if no upgraders ran
+	const FBuiltStruct*							Fixed; // Owned by scratch, null if Loose has any mismatches
 };
 
 class FBatchUpgrader
diff --git a/Engine/Plugins/Experimental/PlainProps/Source/Tests/PlainPropsUpgradeTests.cpp b/Engine/Plugins/Experimental/PlainProps/Source/Tests/PlainPropsUpgradeTests.cpp
--- a/Engine/Plugins/Experimental/PlainProps/Source/Tests/PlainPropsUpgradeTests.cpp
+++ b/Engine/Plugins/Experimental/PlainProps/Source/Tests/PlainPropsUpgradeTests.cpp
@@ -9,6 +9,7 @@
 #include "PlainPropsInternalFormat.h"
 #include "PlainPropsInternalRead.h"
 #include "PlainPropsInternalTest.h"
+#include "PlainPropsInternalUpgrade.h"
 #include "PlainPropsDiff.h"
 #include "PlainPropsBuild.h"
 #include "PlainPropsBuildSchema.h"
@@ -34,25 +35,10 @@ static FSchemaBindings			GSchemas(GDebug);
 static FCustomBindingsBottom	GCustoms(GDebug);
 static FCustomBindingsOverlay	GDeltaCustoms(GCustoms);
 static FBindDeclarations		GDeltaDeclarations(GEnums, GDeltaCustoms, GSchemas);
-static FHistory*				GHistory;
-static FScratchAllocator*		GScratch;
 
 struct FFixture
 {
-	FFixture() : History(MakeHistory(GDeltaDeclarations, GIds)), Vis(GIds, "AnsiStr")
-	{
-		GHistory = &*History;
-		GScratch = &Scratch;
-	}
-
-	~FFixture()
-	{
-		GHistory = nullptr;
-		GScratch = nullptr;
-	}
-	
-	FHistoryPtr	History;
-	FScratchAllocator Scratch;
+	FFixture() : Vis(GIds, "AnsiStr") {}
 	DbgVis::FIdScope Vis;
 };
 
@@ -86,20 +72,96 @@ struct FDeltaRuntime : FDefaultRuntime
 
 ////////////////////////////////////////////////////////////////////////////////////////////////
 
+struct FUnboundDeclarations : IDeclarations
+{
+public:
+	FUnboundDeclarations() : Enums(GDebug) {}
+
+	const FStructDeclaration&				Declare(const char* Scope, const char* Name, std::initializer_list<const char*> MemberNames, TConstArrayView<FMemberSpec> MemberTypes, EMemberPresence Occupancy, FOptionalDeclId Super = NoId)
+	{
+		FDeclId Id = GIds.IndexDeclId(GIds.MakeType(Scope, Name));
+		return *Structs.Emplace(Id, PlainProps::Declare({Id, Super, /* v */ 0, Occupancy, NameMembers(MemberNames), MemberTypes}));	
+	}
+
+private:
+	virtual const FEnumDeclaration*			Find(FEnumId Id) const override		{ return Enums.Find(Id); }
+	virtual const FStructDeclaration*		Find(FStructId Id) const override	{ return Structs.FindChecked(FDeclId(Id)); }
+	virtual FDeclId							Lower(FBindId Id) const override	{ unimplemented(); return LowerCast(Id);	}
+
+	FEnumDeclarations						Enums;
+	TMap<FStructId, FStructDeclarationPtr>	Structs;
+
+	static TArray<FMemberId> NameMembers(std::initializer_list<const char*> Members)
+	{
+		TArray<FMemberId> Out;
+		Out.Reserve(Members.size());
+		for (const char* Member : Members)
+		{
+			Out.Add(GIds.NameMember(Member));
+		}
+		return Out;
+	}
+};
+
+////////////////////////////////////////////////////////////////////////////////////////////////
+
+// FBatchWriter/FBatchSaver helper
+struct FSavedStruct
+{
+	FDeclId						Type;
+	const FBuiltStruct*			Value;
+};
+
+class FBatchWriterBase
+{
+public:
+	FBatchWriterBase(FScratchAllocator& InScratch) : Scratch(InScratch) {}
+		
+	void						SetRoots(TConstArrayView<FUpgradedInstance> Instances)
+	{ 
+		Roots.Empty(Instances.Num());
+		for (FUpgradedInstance Instance : Instances)
+		{
+			check(Instance.Fixed);
+			Roots.Emplace(Instance.Type, Instance.Fixed);
+		}
+	}
+
+protected:
+	TArray<FSavedStruct>		Roots;
+	FScratchAllocator&			Scratch;
+
+	TArray64<uint8>				Write(const IDeclarations& Types) const;
+};
+
+////////////////////////////////////////////////////////////////////////////////////////////////
+
+// For untyped unit tests without load/save/bind layer
+class FBatchWriter : public FBatchWriterBase
+{
+public:
+	using FBatchWriterBase::FBatchWriterBase;
+
+	FBuiltStruct*				Build(FMemberBuilder& Members, const FStructDeclaration& Decl) const	{ return Members.BuildAndReset(Scratch, Decl, GDebug); }
+	void						Save(FDeclId Type, const FBuiltStruct* Value)							{ Roots.Emplace(Type, Value); }
+	TArray64<uint8>				Write() const															{ return FBatchWriterBase::Write(OldTypes); }
+	
+	FUnboundDeclarations		OldTypes;
+};
+
+////////////////////////////////////////////////////////////////////////////////////////////////
+
 // Temporarily bind obsolete test type to save it
 struct ITempBinding { virtual ~ITempBinding() {} };
-
 template<class T>
 struct TTempBinding : ITempBinding, TScopedStructBinding<T, FDefaultRuntime> {};
-
 using FTempBindingPtr = TUniquePtr<ITempBinding>;
 
-inline constexpr uint32 TestMagics[] = { 0xFEEDF00D, 0xABCD1234, 0xDADADAAA, 0x99887766, 0xF0F1F2F3 };
-
-class FBatchSaver
+// For typed unit tests using load/save/bind layer
+class FBatchSaver : public FBatchWriterBase
 {
 public:
-	FBatchSaver(FScratchAllocator& InScratch) : Scratch(InScratch), Types(GEnums, GCustoms, GSchemas) {}
+	using FBatchWriterBase::FBatchWriterBase;
 
 	template<class T>
 	void						Save(T&& Object) 
@@ -109,30 +171,30 @@ public:
 		Roots.Emplace(FDeclId(Id), SaveStruct(&Object, Id, {{GSchemas, GCustoms}, {Scratch}}));
 	}
 
-	void						SetRoots(TArray<FUpgradedInstance>&& In) { Roots = In; }
-	TArray64<uint8>				Write() const;
+	TArray64<uint8>				Write() const { return FBatchWriterBase::Write(GDeltaDeclarations); }
 
 private:
-	TArray<FUpgradedInstance>	Roots;
-	FScratchAllocator&			Scratch;
-	const FBindDeclarations		Types;
 	TArray<FTempBindingPtr>		Obsoletes;
 
 	template<typename T>
 	void BindObsolete(FDualStructId Id)
 	{
-		if (!Types.Find(Id))
+		if (!GDeltaDeclarations.Find(Id))
 		{
 			Obsoletes.Emplace(new TTempBinding<T>);
 		}
 	}
 };
 
-TArray64<uint8> FBatchSaver::Write() const
+////////////////////////////////////////////////////////////////////////////////////////////////
+
+inline constexpr uint32 TestMagics[] = { 0xFEEDF00D, 0xABCD1234, 0xDADADAAA, 0x99887766, 0xF0F1F2F3 };
+
+TArray64<uint8> FBatchWriterBase::Write(const IDeclarations& Types) const
 {
 	// Build partial schemas
 	FSchemasBuilder SchemaBuilders(GIds, Types, Scratch, ESchemaFormat::InMemoryNames);
-	for (FUpgradedInstance Root : Roots)
+	for (FSavedStruct Root : Roots)
 	{
 		SchemaBuilders.NoteStructAndMembers(Root.Type, *Root.Value);
 	}
@@ -153,7 +215,7 @@ TArray64<uint8> FBatchSaver::Write() const
 
 	// Write roots
 	WriteInt(Out, TestMagics[2]);
-	for (FUpgradedInstance Root : Roots)
+	for (FSavedStruct Root : Roots)
 	{
 		WriteInt(/* out */ Tmp, TestMagics[3]);
 		WriteInt(/* out */ Tmp, Writer.GetWriteId(Root.Type).Get().Idx);
@@ -231,7 +293,7 @@ private:
 ////////////////////////////////////////////////////////////////////////////////////////////////
 
 
-static void Run(void (*SaveOld)(FBatchSaver&), void (*TestNew)(FBatchLoader&))
+static void Run(FHistory& History, void (*SaveOld)(FBatchSaver&), void (*TestNew)(FBatchLoader&))
 {
 	// Generate old data
 	TArray64<uint8> OldData;
@@ -252,11 +314,11 @@ static void Run(void (*SaveOld)(FBatchSaver&), void (*TestNew)(FBatchLoader&))
 		FBindDeclarations Declarations(GEnums, GCustoms, GSchemas);
 		CHECK(!Upgrader.MatchSchemas(Declarations));
 		FScratchAllocator Scratch;
-		TArray<FUpgradedInstance> NewRoots = Upgrader.Upgrade(OldBatch.GetRoots(), *GHistory, Scratch);
+		TArray<FUpgradedInstance> NewRoots = Upgrader.Upgrade(OldBatch.GetRoots(), History, Scratch);
 
 		// Write new schemas and instances
 		FBatchSaver NewBatch(Scratch);
-		NewBatch.SetRoots(MoveTemp(NewRoots));
+		NewBatch.SetRoots(NewRoots);
 		NewData = NewBatch.Write();
 	}
 	
@@ -265,6 +327,56 @@ static void Run(void (*SaveOld)(FBatchSaver&), void (*TestNew)(FBatchLoader&))
 	TestNew(NewBatch);
 }
 
+// RunUntyped helper
+struct FTestUpgrades
+{
+	FUnboundDeclarations&	NewTypes;
+	FHistory&				History;
+};
+
+// RunUntyped helper
+struct FUpgradedBatch
+{
+	TConstArrayView<FUpgradedInstance>		Roots;
+	FLooseIndexers&							LooseIds;
+
+	FAnsiString ResolveName(FDeclId Id, FLooseMember Member) const { return GIds.ResolveName(LooseIds.Get(Id).Resolve(Member.GetValue().Meta.Name).Get().Id); }
+};
+
+static void RunUntyped(void (*WriteUntyped)(FBatchWriter&), void (*ChronicleUntyped)(FTestUpgrades), void (*TestUntyped)(FUpgradedBatch))
+{
+	// Generate old data
+	TArray64<uint8> OldData;
+	{
+		FScratchAllocator Scratch;
+		FBatchWriter OldBatch(Scratch);
+		WriteUntyped(OldBatch);
+		OldData = OldBatch.Write();
+	}
+
+	// Register latest declarations and upgrade ops
+	FUnboundDeclarations NewTypes;
+	FHistoryPtr History = MakeHistory(NewTypes, GIds);
+	ChronicleUntyped({NewTypes, *History});
+
+	// Run upgrades
+	FScratchAllocator Scratch;
+	TArray<FUpgradedInstance> Roots;
+	{
+		// Mount old schemas
+		FBatchLoader OldBatch(MakeMemoryView(OldData));
+
+		// Upgrade old instances
+		FBatchUpgrader Upgrader(OldBatch.GetSchemas(), OldBatch.GetRuntimeIds(), GIds);
+		CHECK(!Upgrader.MatchSchemas(NewTypes));
+		
+		Roots = Upgrader.Upgrade(OldBatch.GetRoots(), *History, Scratch);
+	}
+	
+	// Finally test upgraded output
+	TestUntyped({Roots, History->LooseIds});
+}
+
 ////////////////////////////////////////////////////////////////////////////////////////////////
 
 // Tests that everything was read
@@ -698,17 +810,17 @@ PP_REFLECT_STRUCT(PlainProps::Upgrade::Test, FRangesB, void, Leaves, String, Poi
 template<class T>
 struct TScopedStructUpgradeBinding 
 {
-	TScopedStructUpgradeBinding()
+	explicit TScopedStructUpgradeBinding(FHistory& Hist) : History(Hist)
 	{
-		Record(static_cast<TChroniclesOf<T>>(FChronicles(/* out */ *GHistory, Binding.DeclId)));
+		Record(static_cast<TChroniclesOf<T>>(FChronicles(/* out */ History, Binding.DeclId)));
 	}
 	
 	~TScopedStructUpgradeBinding()
 	{
-		EraseChroniclesOf(*GHistory, Binding.DeclId);
+		EraseChroniclesOf(History, Binding.DeclId);
 	}
-
-	TScopedStructBinding<T, FDefaultRuntime> Binding;
+	FHistory&									History;
+	TScopedStructBinding<T, FDefaultRuntime>	Binding;
 };
 
 ////////////////////////////////////////////////////////////////////////////////////////////////
@@ -723,22 +835,24 @@ TEST_CASE_NAMED(FPlainPropsUpgradeTest, "System::PlainProps::Upgrade", "[PlainPr
 
 	SECTION("Basics")
 	{
-		TScopedStructUpgradeBinding<FBasic> Basic;
+		FHistoryPtr History = MakeHistory(GDeltaDeclarations, GIds);
+		TScopedStructUpgradeBinding<FBasic> Basic(*History);
 		TScopedEnumDeclaration<EFlag, EEnumMode::Flag, FDefaultRuntime> Flag;
 		TScopedStructBinding<FPoint, FDefaultRuntime> Point;
 
-		Run([](FBatchSaver& Batch)
-		{
-			Batch.Save(FBasicOld{
-				.Widen=1, .Rename=2, .Reorder=3, .Drop=4, .Unchanged=5, .FlagA=false, .FlagB=true, 
-				.PointX=6.f, .PointY=7.f, .Guid=FGuid(0x01234567, 0x89ABCDEF,0,0xFFFFFFFF)});
-		}, 
-		[](FBatchLoader& Batch)
-		{
-			CHECK(Batch.Load<FBasic>(0) == FBasic{
-				.Widen=1, .Reorder=3, .Renamed=2, .Unchanged=5, .Flags=EFlag::B, 
-				.Point={6.f, 7.f}, .GuidAB=0x0123456789ABCDEF, .GuidCD=0xFFFFFFFF});
-		});
+		Run(*History,
+			[](FBatchSaver& Batch)
+			{
+				Batch.Save(FBasicOld{
+					.Widen=1, .Rename=2, .Reorder=3, .Drop=4, .Unchanged=5, .FlagA=false, .FlagB=true, 
+					.PointX=6.f, .PointY=7.f, .Guid=FGuid(0x01234567, 0x89ABCDEF,0,0xFFFFFFFF)});
+			}, 
+			[](FBatchLoader& Batch)
+			{
+				CHECK(Batch.Load<FBasic>(0) == FBasic{
+					.Widen=1, .Reorder=3, .Renamed=2, .Unchanged=5, .Flags=EFlag::B, 
+					.Point={6.f, 7.f}, .GuidAB=0x0123456789ABCDEF, .GuidCD=0xFFFFFFFF});
+			});
 	}
 
 	SECTION("Ranges")
@@ -757,8 +871,120 @@ TEST_CASE_NAMED(FPlainPropsUpgradeTest, "System::PlainProps::Upgrade", "[PlainPr
 	{
 	}
 
-	SECTION("Errors")
+	SECTION("Mismatches")
 	{
+		RunUntyped([](FBatchWriter& Batch)
+		{
+			const FStructDeclaration& InputDecl = Batch.OldTypes.Declare("Test", "Input", {"A","B","X"}, {SpecF64, SpecF32, SpecF32}, EMemberPresence::AllowSparse);
+			const FStructDeclaration& MismatchDecl = Batch.OldTypes.Declare("Test", "Mismatch", 
+				{"UnusedOld", "Unknown", "Mistyped", "Ok", "Inputs"}, 
+				{SpecS32,	SpecS32,	SpecS32,	SpecS32, FMemberSpec(InputDecl.Id)}, EMemberPresence::AllowSparse);
+
+			FMemberBuilder Members;
+			Members.Add(MismatchDecl.MemberNames[1], 123);	// Unknown
+			Members.Add(MismatchDecl.MemberNames[2], 456);	// Mistyped
+			Members.Add(MismatchDecl.MemberNames[3], 789);	// Ok
+
+			FMemberBuilder Inputs;
+			Inputs.Add(GIds.IndexMemberLiteral("A"), 1.0);	// Wrong type
+			Inputs.Add(GIds.IndexMemberLiteral("B"), 2.0f); // Ok
+			Inputs.Add(GIds.IndexMemberLiteral("X"), 3.0f); // Surplus
+															// C is missing			
+			Members.AddStruct(MismatchDecl.MemberNames[4], InputDecl.Id, Batch.Build(Inputs, InputDecl));
+
+			Batch.Save(MismatchDecl.Id, Batch.Build(Members, MismatchDecl));
+		}, 
+		[](FTestUpgrades Upgrades)
+		{
+			// First declare types
+			const FStructDeclaration& InputDecl = Upgrades.NewTypes.Declare("Test", "Input", {"A","B","C"}, {SpecF32, SpecF32, SpecF32}, EMemberPresence::AllowSparse);
+			const FStructDeclaration& MismatchDecl = Upgrades.NewTypes.Declare("Test", "Mismatch", 
+				{"Duplicate", "Mistyped", "Ok",	"Inputs"}, 
+				{SpecS32,		SpecS64, SpecS32, FMemberSpec(InputDecl.Id)}, EMemberPresence::AllowSparse);
+		
+			FChronicles MismatchChronicles(Upgrades.History, MismatchDecl.Id);
+			IOps& MismatchOps = *MismatchChronicles.CreateVersion({"Test", "Mismatch"}, 0);
+
+			// Add two ops which both reemit Ok but also emit one Duplicate
+			MismatchOps.Add({"Ok"_S32}, {"Ok"_S32, "Duplicate"_S32}, [](FIns In, FOuts Out, FContext Ctx)
+			{
+				Out[0].S32() = In[0].S32();
+				Out[1].S32() = 111;
+			});
+
+			MismatchOps.Add({"Ok"_S32}, {"Ok"_S32, "Duplicate"_S32}, [](FIns In, FOuts Out, FContext Ctx)
+			{
+				Out[0].S32() = In[0].S32();
+				Out[1].S32() = 222;
+			});
+			
+			FStructId InputId = MismatchOps.NameStruct({"Test", "Input"});
+			FStructMemberNames ABC = MismatchOps.NameMembers(InputId, {"A"_F32, "B"_F32, "C"_F32});
+			MismatchOps.Add({"Inputs"_Struct(ABC)}, {}, [](FIns In, FOuts Out, FContext)
+				{
+					CHECK(false); // Shouldn't execute op with mismatching inputs
+				});
+
+			MismatchChronicles.CommitVersion(&MismatchOps);
+
+			// Test.Input chronicles are missing X -> C rename and A's F64 ->F32 conversion
+			FChronicles InputChronicles(Upgrades.History, InputDecl.Id);
+			IOps& InputOps = *InputChronicles.CreateVersion({"Test", "Input"}, 0);
+			InputChronicles.CommitVersion(&InputOps);
+		}, 
+		[](FUpgradedBatch Batch)
+		{
+			CHECK(Batch.Roots.Num() == 1);
+			CHECK(Batch.Roots[0].Fixed == nullptr);
+			CHECK(Batch.Roots[0].Loose.Num() == 6);
+		
+			FDeclId MismatchId = Batch.Roots[0].Type;
+			TConstArrayView<FLooseMember> Members = Batch.Roots[0].Loose;
+
+			//TArray<FAnsiString> MemberNames;
+			//for (FLooseMember Member : Members)
+			//{
+			//	MemberNames.Emplace(Batch.ResolveName(MismatchId, Member));
+			//}
+
+			// Matching members come first
+			// Mismatching member order doesn't matter, this is a hacky test of a temporary / ad-hoc API
+			CHECK(Batch.ResolveName(MismatchId, Members[0]) == "Ok");
+			CHECK(Batch.ResolveName(MismatchId, Members[1]) == "Duplicate"); 
+			CHECK(Batch.ResolveName(MismatchId, Members[2]) == "Unknown");
+			CHECK(Batch.ResolveName(MismatchId, Members[3]) == "Mistyped");
+			CHECK(Batch.ResolveName(MismatchId, Members[4]) == "Duplicate");
+			CHECK(Batch.ResolveName(MismatchId, Members[5]) == "Inputs");
+
+			// 222 is given first despite 111 being output first due to current greedy execution that might change
+			CHECK(Members[0].GetValue().Meta.Mismatch == EMismatch::No);
+			CHECK(Members[0].GetValue().Data.Arithmetic == 789);
+			CHECK(Members[1].GetValue().Meta.Mismatch == EMismatch::Duplicate);
+			CHECK(Members[1].GetValue().Data.Arithmetic == 222);				
+			CHECK(Members[2].GetValue().Meta.Mismatch == EMismatch::Surplus);
+			CHECK(Members[2].GetValue().Data.Arithmetic == 123);
+			CHECK(Members[3].GetValue().Meta.Mismatch == EMismatch::Type);
+			CHECK(Members[3].GetValue().Data.Arithmetic == 456);
+			CHECK(Members[4].GetValue().Meta.Mismatch == EMismatch::Duplicate);
+			CHECK(Members[4].GetValue().Data.Arithmetic == 111);
+			CHECK(Members[5].GetValue().Meta.Mismatch == EMismatch::Inputs);
+
+			FLooseMember Inputs = Members[5];
+			CHECK(Inputs.GetType().IsStruct());
+
+			FDeclId InputId = Inputs.GetType().InnermostId.Get().AsStructDeclId();
+			TConstArrayView<FLooseMember> Inners(Inputs.GetValue().Data.Struct.Loose, Inputs.GetValue().Meta.Num);
+			CHECK(Inners.Num() == 4);			
+			CHECK(Batch.ResolveName(InputId, Inners[0]) == "B");
+			CHECK(Batch.ResolveName(InputId, Inners[1]) == "A");
+			CHECK(Batch.ResolveName(InputId, Inners[2]) == "X");
+			CHECK(Batch.ResolveName(InputId, Inners[3]) == "C");
+
+			CHECK(Inners[0].GetValue().Meta.Mismatch == EMismatch::No);
+			CHECK(Inners[1].GetValue().Meta.Mismatch == EMismatch::Type);
+			CHECK(Inners[2].GetValue().Meta.Mismatch == EMismatch::Surplus);
+			CHECK(Inners[3].GetValue().Meta.Mismatch == EMismatch::Missing);
+		});
 	}
 
 	SECTION("Defaults")
