PlainProps fix static shutdown issue

#rnx
#Synced-CL 49069666

diff --git a/Engine/Plugins/Experimental/PlainProps/Source/Private/PlainPropsUeCoreBindings.cpp b/Engine/Plugins/Experimental/PlainProps/Source/Private/PlainPropsUeCoreBindings.cpp
--- a/Engine/Plugins/Experimental/PlainProps/Source/Private/PlainPropsUeCoreBindings.cpp
+++ b/Engine/Plugins/Experimental/PlainProps/Source/Private/PlainPropsUeCoreBindings.cpp
@@ -52,12 +52,14 @@ void FTransformBinding::Save(FMemberBuilder& Dst, const FTransform& Src, const F
 {
 	static_assert(std::is_same_v<decltype(FTransform().GetTranslation().X), double>);
 	
-	static const struct
+	// Lazily cache declarations on first save, then hoist them to the stack
+	if (!VectorDeclaration)
 	{
-		FStructDeclarationPtr Vector;
-		FStructDeclarationPtr Quat;
+		VectorDeclaration = &Ctx.Schemas.GetDeclaration(VectorId);
+		QuatDeclaration = &Ctx.Schemas.GetDeclaration(QuatId);
 	}
-	Decls = { &Ctx.Schemas.GetDeclaration(VectorId), &Ctx.Schemas.GetDeclaration(QuatId) };
+	const FStructDeclaration& VectorDecl = *VectorDeclaration;
+	const FStructDeclaration& QuatDecl = *QuatDeclaration;
 
 	FDenseMemberBuilder Inner = { Ctx.Scratch, Ctx.Schemas.GetDebug() };
 
@@ -69,24 +71,24 @@ void FTransformBinding::Save(FMemberBuilder& Dst, const FTransform& Src, const F
 	{
 		if (T != Default->GetTranslation())
 		{
-			Dst.AddStruct(MemberIds[(uint8)EMember::Translate], VectorId, Inner.BuildHomogeneous(*Decls.Vector, T.X, T.Y, T.Z));
+			Dst.AddStruct(MemberIds[(uint8)EMember::Translate], VectorId, Inner.BuildHomogeneous(VectorDecl, T.X, T.Y, T.Z));
 		}
 
 		if (R != Default->GetRotation())
 		{
-			Dst.AddStruct(MemberIds[(uint8)EMember::Rotate], QuatId, Inner.BuildHomogeneous(*Decls.Quat, R.X, R.Y, R.Z, R.W));
+			Dst.AddStruct(MemberIds[(uint8)EMember::Rotate], QuatId, Inner.BuildHomogeneous(QuatDecl, R.X, R.Y, R.Z, R.W));
 		}
 
 		if (S != Default->GetScale3D())
 		{
-			Dst.AddStruct(MemberIds[(uint8)EMember::Scale], VectorId, Inner.BuildHomogeneous(*Decls.Vector, S.X, S.Y, S.Z));
+			Dst.AddStruct(MemberIds[(uint8)EMember::Scale], VectorId, Inner.BuildHomogeneous(VectorDecl, S.X, S.Y, S.Z));
 		}
 	}
 	else
 	{	
-		Dst.AddStruct(MemberIds[(uint8)EMember::Translate], VectorId, Inner.BuildHomogeneous(*Decls.Vector, T.X, T.Y, T.Z));		
-		Dst.AddStruct(MemberIds[(uint8)EMember::Rotate], QuatId, Inner.BuildHomogeneous(*Decls.Quat, R.X, R.Y, R.Z, R.W));
-		Dst.AddStruct(MemberIds[(uint8)EMember::Scale], VectorId, Inner.BuildHomogeneous(*Decls.Vector, S.X, S.Y, S.Z));
+		Dst.AddStruct(MemberIds[(uint8)EMember::Translate], VectorId, Inner.BuildHomogeneous(VectorDecl, T.X, T.Y, T.Z));		
+		Dst.AddStruct(MemberIds[(uint8)EMember::Rotate], QuatId, Inner.BuildHomogeneous(QuatDecl, R.X, R.Y, R.Z, R.W));
+		Dst.AddStruct(MemberIds[(uint8)EMember::Scale], VectorId, Inner.BuildHomogeneous(VectorDecl, S.X, S.Y, S.Z));
 	}
 }
 
diff --git a/Engine/Plugins/Experimental/PlainProps/Source/Public/PlainPropsUeCoreBindings.h b/Engine/Plugins/Experimental/PlainProps/Source/Public/PlainPropsUeCoreBindings.h
--- a/Engine/Plugins/Experimental/PlainProps/Source/Public/PlainPropsUeCoreBindings.h
+++ b/Engine/Plugins/Experimental/PlainProps/Source/Public/PlainPropsUeCoreBindings.h
@@ -643,6 +643,8 @@ struct FTransformBinding : ICustomBinding
 	const FMemberId					MemberIds[3];
 	const FDualStructId				VectorId;
 	const FDualStructId				QuatId;
+	mutable FStructDeclarationPtr	VectorDeclaration;
+	mutable FStructDeclarationPtr	QuatDeclaration;
 
 	template<class Ids>
 	FTransformBinding(TCustomSpecifier<Ids, 3>& Spec)
