PlainProps: Enable loading blueprint graphs in linker round trip test by adding support for updating a load batch with new load plans for dynamic structs.

Result:
Loaded 3437 plainprops packages in 15.91 s
Loaded 3437 baseline packages in 19.94 s
Ignored 786 package diff(s) in 3437 package(s)
Detected 33 package diff(s) in 3437 package(s)
Loaded 132133 exports, 202 load plan updates

#rnx

#Synced-CL 49027723

diff --git a/Engine/Plugins/Experimental/PlainProps/Source/Private/PlainPropsLoad.cpp b/Engine/Plugins/Experimental/PlainProps/Source/Private/PlainPropsLoad.cpp
--- a/Engine/Plugins/Experimental/PlainProps/Source/Private/PlainPropsLoad.cpp
+++ b/Engine/Plugins/Experimental/PlainProps/Source/Private/PlainPropsLoad.cpp
@@ -72,6 +72,9 @@ public:
 		check(OffsetWidth == GetOffsetWidth());
 	}
 
+	friend bool operator==(FLoadStructPlan, FLoadStructPlan) = default;
+	friend bool operator!=(FLoadStructPlan, FLoadStructPlan) = default;
+
 	bool						IsSchema() const		{ return (Handle & SchemaBit) == SchemaBit; }
 	bool						IsSparseSchema() const	{ return (Handle & SparseSchemaMask) == SparseSchemaMask; }
 	bool						IsMemcpy() const		{ return (Handle & LoMask) == MemcpyMask; }
@@ -185,11 +188,13 @@ struct FLoadIdMapping
 class FLoadIdBinder
 {
 public:
-	FLoadIdBinder(TConstArrayView<FDeclId> Ids, FDebugIds Dbg)
+	FLoadIdBinder(TConstArrayView<FDeclId> Ids, uint32 NumPlans, FDebugIds Dbg)
 	: DeclIds(Ids)
-	, NextLoadIdx(static_cast<uint32>(Ids.Num()))
+	, NextLoadIdx(NumPlans)
 	, Debug(Dbg)
-	{}
+	{
+		check(NextLoadIdx >= static_cast<uint32>(Ids.Num()));
+	}
 
 	FStructSchemaId					BindLoadId(FStructSchemaId ReadId, FBindId Id)
 	{
@@ -236,11 +241,19 @@ struct FLoadIdBinderDummy
 
 ////////////////////////////////////////////////////////////////////////////
 
+struct FLateLoadPlans
+{
+	FLateLoadPlans* Next;
+	uint8 Data[0];
+};
+
 struct FLoadBatch
 {
-	FSchemaBatchId			BatchId; // Needed to access schemas for custom struct loading
-	uint32					NumReadSchemas;
-	uint32					NumPlans;
+	FSchemaBatchId			BatchId;			// Needed to access schemas for custom struct loading
+	uint32					NumReadSchemas;		// Number of saved struct schemas
+	uint32					NumPlans;			// NumReadSchemas + NumMappings (for bound type-erased structs)
+	uint32					MaxPlans;			// MaxPlans required to handle late bound structs with potentially type-erased inners
+	FLateLoadPlans*			LatePlans;			// Data for late bound structs
 	FLoadStructPlan			Plans[0];
 
 	FLoadStructPlan			operator[](FStructSchemaId LoadId) const
@@ -253,14 +266,22 @@ struct FLoadBatch
 	{
 		check(LoadId.Idx < NumPlans);
 		static_assert(alignof(FLoadStructPlan) >= alignof(FStructSchemaId));
-		const FStructSchemaId* SaveIds = reinterpret_cast<const FStructSchemaId*>(Plans + NumPlans);
+		const FStructSchemaId* SaveIds = reinterpret_cast<const FStructSchemaId*>(Plans + MaxPlans);
 		return LoadId.Idx < NumReadSchemas ? LoadId : SaveIds[LoadId.Idx - NumReadSchemas];
 	}
 };
 
 void FLoadBatchDeleter::operator()(FLoadBatch* Batch) const
 {
-	FMemory::Free(Batch);
+	if (Batch)
+	{
+		for (FLateLoadPlans *It = Batch->LatePlans, *Next; It; It = Next)
+		{
+			Next = It->Next;
+			FMemory::Free(It);
+		}
+		FMemory::Free(Batch);
+	}
 }
 
 ////////////////////////////////////////////////////////////////////////////
@@ -301,7 +322,21 @@ struct FNoopLoadBinding final : FLoadBindingBase
 	virtual void LoadCustom(void* Dst, FStructLoadView Src, ECustomLoadMethod Method) const override {}
 };
 static FNoopLoadBinding GLoadNoop;
+struct FLateLoadBinding final : FLoadBindingBase
+{
+	virtual void LoadCustom(void* Dst, FStructLoadView Src, ECustomLoadMethod) const override { checkf(false, TEXT("Can't load missing/late struct binding with load id %d"), Src.Schema.LoadId.Idx); }
+};
+static FLateLoadBinding GLoadNotBoundError;
+
+inline static bool IsBound(FLoadStructPlan Plan)
+{
+	return Plan != FLoadStructPlan(GLoadNotBoundError);
+}
 
+bool HasLoadPlan(FStructSchemaId LoadId, const FLoadBatch& Batch)
+{
+	return IsBound(Batch[LoadId]);
+}
 
 template<class T>
 static TOptional<TConstArrayView<T>> GetContiguousSubset(TConstArrayView<T> View, const TBitArray<>& Keep)
@@ -395,22 +430,81 @@ struct FLoadPlanner
 					CreatePlan(LoadIds.GetMapping(Idx)); // LoadIds might grow
 				}
 
-				// Verify that all UnboundSaveIds were bound by some load plan
+				// Strip all UnboundSaveIds that were bound by some load plan,
+				// and set set the load plan to GLoadTypeErasedError to mark it as bound.
 				for (int32 Idx = 0; !UnboundSaveIds.IsEmpty() && Idx < LoadIds.NumMappings(); ++Idx)
 				{
-					UnboundSaveIds.Remove(LoadIds.GetMapping(Idx).ReadId);
+					FStructSchemaId ReadId = LoadIds.GetMapping(Idx).ReadId;
+					Plans[ReadId.Idx] = FLoadStructPlan(GLoadTypeErasedError);
+					UnboundSaveIds.Remove(ReadId);
 				}
 			}
 
-			for (FStructSchemaId UnboundSaveId : UnboundSaveIds)
-			{
-				checkf(false, TEXT("Unbound struct '%s' can't be loaded"), *Schemas.GetDebug().Print(RuntimeIds[UnboundSaveId.Idx]));
-			}
+			// TODO: Bring back check for verifying that all structs are bound when we can reliably identify late-bound structs
+			// for (FStructSchemaId UnboundSaveId : UnboundSaveIds)
+			// {
+			// 	checkf(false, TEXT("Unbound struct '%s' can't be loaded"), *Schemas.GetDebug().Print(RuntimeIds[UnboundSaveId.Idx]));
+			// }
 
 			// Allocate load batch, copy plans and subset schemas, and fixup subset schema plans
 			return CreateBatch();
 		}
 	}
+
+	void UpdatePlans(FLoadBatch& Batch)
+	{
+		check(Batch.NumReadSchemas == RuntimeIds.Num());
+		check(Batch.NumPlans >= Batch.NumReadSchemas);
+
+		// Initialize temp plans to GLoadNotBoundError.
+		// The temp SubsetSchemaSizes will only be accessed for newly bound plans.
+		const uint32 NumMappings = Batch.NumPlans - Batch.NumReadSchemas;
+		Plans.Init(FLoadStructPlan(GLoadNotBoundError), Batch.NumPlans);
+		SubsetSchemaSizes.SetNumUninitialized(Batch.NumPlans);
+
+		// Eagerly create load plans for all unbound read schemas
+		for (uint32 Idx = 0; Idx < Batch.NumReadSchemas; ++Idx)
+		{
+			if (!IsBound(Batch.Plans[Idx]))
+			{
+				FLoadIdMapping Mapping;
+				Mapping.ReadId = { Idx };
+				Mapping.LoadId = { Idx };
+				Mapping.Id = FBindId(RuntimeIds[Idx]);
+				CreatePlan(Mapping);
+				if (IsBound(Plans[Mapping.ReadId.Idx]))
+				{
+					UnboundSaveIds.Remove(Mapping.ReadId);
+				}
+			}
+		}
+
+		// Make load plans for type-erased/ExplicitBindName structs needed by the new load plans
+		if (LoadIds.NumMappings() > 0)
+		{
+			Plans.Reserve(Batch.NumPlans + LoadIds.NumMappings());
+			SubsetSchemaSizes.Reserve(Batch.NumPlans + LoadIds.NumMappings());
+			for (int32 Idx = 0; Idx < LoadIds.NumMappings(); ++Idx)
+			{
+				check(LoadIds.GetMapping(Idx).LoadId.Idx == static_cast<uint32>(Plans.Num()));
+				Plans.AddUninitialized();
+				SubsetSchemaSizes.AddUninitialized();
+				CreatePlan(LoadIds.GetMapping(Idx)); // LoadIds might grow
+			}
+
+			// Remove all UnboundSaveIds for type erased structs that have now been bound by some load plan
+			for (int32 Idx = 0; !UnboundSaveIds.IsEmpty() && Idx < LoadIds.NumMappings(); ++Idx)
+			{
+				FStructSchemaId ReadId = LoadIds.GetMapping(Idx).ReadId;
+				Plans[ReadId.Idx] = FLoadStructPlan(GLoadTypeErasedError);
+				UnboundSaveIds.Remove(ReadId);
+			}
+		}
+
+		// Update load batch, copy new plans and subset schemas, and fixup new subset schema plans
+		UpdateBatch(Batch);
+	}
+
 private:
 
 	FLoadBatchPtr CreateBatch() const
@@ -418,23 +512,26 @@ private:
 		const uint32 NumPlans = Plans.Num();
 		const uint32 NumMappings = LoadIds.NumMappings();
 		const uint32 NumReadSchemas = RuntimeIds.Num();
+		const uint32 NumSlack = UnboundSaveIds.Num();
+		const uint32 MaxPlans = NumPlans + NumSlack;
 		check(NumPlans == NumReadSchemas + NumMappings);
 
 		// Allocate plan, init header and copy plans
-		SIZE_T Bytes = sizeof(FLoadBatch) + sizeof(FLoadStructPlan) * NumPlans
-					 + Align(sizeof(FStructSchemaId) * NumMappings, alignof(FSchemaBinding))
+		SIZE_T Bytes = sizeof(FLoadBatch) + sizeof(FLoadStructPlan) * MaxPlans
+					 + Align(sizeof(FStructSchemaId) * (NumMappings + NumSlack), alignof(FSchemaBinding))
 					 + SubsetSchemaData.Num();
-		FLoadBatch Header = { BatchId, NumReadSchemas, NumPlans };
+		FLoadBatch Header = { BatchId, NumReadSchemas, NumPlans, MaxPlans };
 		FLoadBatch* Out = new (FMemory::Malloc(Bytes)) FLoadBatch{Header};
 		FMemory::Memcpy(Out->Plans, Plans.GetData(), sizeof(FLoadStructPlan) * NumPlans);
 
 		// Copy LoadId -> ReadId mapping so custom-bound mapped plans can form FReadSchemaHandle and FStructView
-		FStructSchemaId* OutReadId = reinterpret_cast<FStructSchemaId*>(Out->Plans + NumPlans);
+		FStructSchemaId* OutReadId = reinterpret_cast<FStructSchemaId*>(Out->Plans + MaxPlans);
 		for (uint32 Idx = 0; Idx < NumMappings; ++Idx)
 		{
 			*OutReadId++ = LoadIds.GetMapping(Idx).ReadId;
 			check(OutReadId[-1].Idx < NumReadSchemas);
 		}
+		OutReadId += NumSlack;
 
 		// Copy cloned subset schemas and patch up plan pointers
 		if (SubsetSchemaData.Num() > 0)
@@ -446,33 +543,96 @@ private:
 			const uint8* It = static_cast<const uint8*>(OutSubsetData);
 			for (uint32 Idx = 0; Idx < NumPlans; ++Idx)
 			{
-				if (int32 Size = SubsetSchemaSizes[Idx])
+				It += UpdatePlanPointer(*Out, Idx, It);
+			}
+			check(It == (uint8*)OutSubsetData + SubsetSchemaData.Num());
+			check(It == (uint8*)Out + Bytes);
+		}
+		else
+		{
+			check((uint8*)Align(OutReadId, alignof(FSchemaBinding)) == (uint8*)Out + Bytes);
+		}
+
+		return FLoadBatchPtr(Out);
+	}
+
+	void UpdateBatch(FLoadBatch& Batch) const
+	{
+		const uint32 NumPlans = Plans.Num();
+		const uint32 NumMappings = LoadIds.NumMappings();
+		check(NumPlans <= Batch.MaxPlans);
+		check(NumMappings <= Batch.MaxPlans - Batch.NumPlans);
+		check(NumMappings == NumPlans - Batch.NumPlans);
+
+		// Copy new LoadId -> ReadId mappings so custom-bound mapped plans can form FReadSchemaHandle and FStructView
+		FStructSchemaId* OutReadId = reinterpret_cast<FStructSchemaId*>(Batch.Plans + Batch.MaxPlans);
+		OutReadId += Batch.NumPlans - Batch.NumReadSchemas; // skip past existing mappings
+		for (uint32 Idx = 0; Idx < NumMappings; ++Idx)
+		{
+			*OutReadId++ = LoadIds.GetMapping(Idx).ReadId;
+			check(OutReadId[-1].Idx < Batch.NumReadSchemas);
+		}
+
+		// Copy cloned subset schemas and patch up new plan pointers
+		if (SubsetSchemaData.Num() > 0)
+		{
+			// Allocate new subset data and add it to the linked list of dynamic allocations
+			SIZE_T Bytes = Align(sizeof(FLateLoadPlans) + SubsetSchemaData.Num(), alignof(FSchemaBinding));
+			FLateLoadPlans LatePlans = { Batch.LatePlans };
+			FLateLoadPlans* Ptr = new (FMemory::Malloc(Bytes)) FLateLoadPlans{ LatePlans };
+			Batch.LatePlans = Ptr;
+
+			void* OutSubsetData = Align(&Ptr[1], alignof(FSchemaBinding));
+			FMemory::Memcpy(OutSubsetData, SubsetSchemaData.GetData(), SubsetSchemaData.Num());
+
+			// Update the existing and added slots with actual subset schema pointers
+			const uint8* It = static_cast<const uint8*>(OutSubsetData);
+			for (uint32 Idx = 0; Idx < NumPlans; ++Idx)
+			{
+				if (IsBound(Plans[Idx]))
 				{
-					if (Plans[Idx].IsSchema())
+					if (uint32 Size = UpdatePlanPointer(Batch, Idx, It))
 					{
-						check(IsAligned(Size, alignof(FSchemaBinding)));
-						bool bSparse = Plans[Idx].IsSparseSchema();
-						Out->Plans[Idx] = FLoadStructPlan(*reinterpret_cast<const FSchemaBinding*>(It), ELeafWidth::B32, bSparse);
+						It += Size;
 					}
 					else
 					{
-						check(Plans[Idx].IsCustom());
-						check(IsAligned(Size, alignof(FCustomLoadPlan)));
-						Out->Plans[Idx] = FLoadStructPlan(*reinterpret_cast<const FCustomLoadPlan*>(It));
+						Batch.Plans[Idx] = Plans[Idx];
 					}
-					It += Size;
 				}
 			}
 			check(It == (uint8*)OutSubsetData + SubsetSchemaData.Num());
-			check(It == (uint8*)Out + Bytes);
+			check(It == (uint8*)Ptr + Bytes);
 		}
 
-		return FLoadBatchPtr(Out);
+		Batch.NumPlans = NumPlans;
+	}
+
+	[[nodiscard]] inline uint32 UpdatePlanPointer(FLoadBatch& Batch, uint32 Idx, const uint8* Src) const
+	{
+		uint32 Size = SubsetSchemaSizes[Idx];
+		if (Size)
+		{
+			if (Plans[Idx].IsSchema())
+			{
+				check(IsAligned(Size, alignof(FSchemaBinding)));
+				bool bSparse = Plans[Idx].IsSparseSchema();
+				Batch.Plans[Idx] = FLoadStructPlan(*reinterpret_cast<const FSchemaBinding*>(Src), ELeafWidth::B32, bSparse);
+			}
+			else
+			{
+				check(Plans[Idx].IsCustom());
+				check(IsAligned(Size, alignof(FCustomLoadPlan)));
+				Batch.Plans[Idx] = FLoadStructPlan(*reinterpret_cast<const FCustomLoadPlan*>(Src));
+			}
+		}
+		return Size;
 	}
 
 	void CreatePlan(FLoadIdMapping Mapping)
 	{
-		int32 SubsetSchemaOffset = SubsetSchemaData.Num();	
+		static_assert(alignof(FCustomLoadPlan) == alignof(FSchemaBinding), "Alignment equality is expected by SubsetSchemaData");
+		const int32 SubsetSchemaOffset = SubsetSchemaData.Num();
 		Plans[Mapping.LoadId.Idx] = CreatePlanInner(Mapping);
 		SubsetSchemaSizes[Mapping.LoadId.Idx] = static_cast<uint32>(SubsetSchemaData.Num() - SubsetSchemaOffset);
 	}
@@ -523,6 +683,7 @@ private:
 		const uint32 Num = static_cast<uint32>(Inners.Num());
 		check(Num);
 		check(Schema.NumInnerSchemas > 0);
+		check(IsAligned(SubsetSchemaData.Num(), alignof(FCustomLoadPlan)));
 
 		// Calculate padding and size needed
 		uint32 Size = offsetof(FCustomLoadPlan, LoadIds) + sizeof(FStructSchemaId) * Num;
@@ -631,9 +792,12 @@ private:
 				return MakeLoadPlan(From, *To, ToMemberIds);	
 			}
 			
-			// Type-erased structs
+			// Type-erased structs are expected to be handled right away in a second iteration
+			//     when processing added mappings.
+			// Late-bound structs are expected to be handled during loading so that the batch
+			//     can be updated with a new load plan before the struct data is to be loaded.
 			UnboundSaveIds.Add(Mapping.ReadId);
-			return FLoadStructPlan(GLoadTypeErasedError);
+			return FLoadStructPlan(GLoadNotBoundError);
 		}
 
 		return FLoadStructPlan(GLoadNoop);
@@ -690,13 +854,13 @@ private:
 		check(To.NumMembers >= From.NumMembers);
 		check(IsAligned(SubsetSchemaData.Num(), alignof(FSchemaBinding)));
 
-		int32 OutPos = SubsetSchemaData.Num();
-		
 		// Allocate and init header 
 		bool bSkipSuperSchema = SkipDeclaredSuperSchema(From.Inheritance);
 		FSchemaBinding Header = { To.DeclId, From.NumMembers, From.NumInnerSchemas - NumEnums - bSkipSuperSchema, From.NumRangeTypes };
+		int32 Offset = SubsetSchemaData.Num();
 		SubsetSchemaData.AddUninitialized(Header.CalculateSize());
-		FSchemaBinding* Schema = new (&SubsetSchemaData[OutPos]) FSchemaBinding {Header};
+		FSchemaBinding* Schema = new (&SubsetSchemaData[Offset]) FSchemaBinding {Header};
+		check(IsAligned(Schema, alignof(FSchemaBinding)));
 		
 		// Copy subset of member bindings...
 		FMemberVisitor ToIt(To);
@@ -758,6 +922,7 @@ private:
 		uint32 Size = To.CalculateSize();
 		SubsetSchemaData.AddUninitialized(Size);
 		FSchemaBinding* Schema = reinterpret_cast<FSchemaBinding*>(&SubsetSchemaData[SubsetSchemaData.Num() - Size]);
+		check(IsAligned(Schema, alignof(FSchemaBinding)));
 		FMemory::Memcpy(Schema, &To, Size);
 
 		// Replace inner bind ids with batch load ids
@@ -849,11 +1014,18 @@ FLoadBatchPtr CreateLoadPlans(FSchemaBatchId BatchId, const FCustomBindings& Cus
 	else
 	{
 		TConstArrayView<FDeclId> DeclIds(static_cast<const FDeclId*>(RuntimeIds.GetData()), RuntimeIds.Num());
-		FLoadIdBinder LoadIds(DeclIds, Schemas.GetDebug());
+		FLoadIdBinder LoadIds(DeclIds, static_cast<uint32>(DeclIds.Num()), Schemas.GetDebug());
 		return FLoadPlanner{BatchId, Customs, Schemas, RuntimeIds, LoadIds}.CreatePlans();
 	}
 }
 
+void UpdateLoadPlans(FLoadBatch& Batch, const FCustomBindings& Customs, const FSchemaBindings& Schemas, TConstArrayView<FStructId> RuntimeIds)
+{
+	TConstArrayView<FDeclId> DeclIds(static_cast<const FDeclId*>(RuntimeIds.GetData()), RuntimeIds.Num());
+	FLoadIdBinder LoadIds(DeclIds, Batch.NumPlans, Schemas.GetDebug());
+	FLoadPlanner{Batch.BatchId, Customs, Schemas, RuntimeIds, LoadIds}.UpdatePlans(Batch);
+}
+
 ////////////////////////////////////////////////////////////////////////////
 
 inline void SetBit(uint8& Out, uint8 Idx, bool bValue)
diff --git a/Engine/Plugins/Experimental/PlainProps/Source/Public/PlainPropsLoad.h b/Engine/Plugins/Experimental/PlainProps/Source/Public/PlainPropsLoad.h
--- a/Engine/Plugins/Experimental/PlainProps/Source/Public/PlainPropsLoad.h
+++ b/Engine/Plugins/Experimental/PlainProps/Source/Public/PlainPropsLoad.h
@@ -22,7 +22,9 @@ using FLoadBatchPtr = TUniquePtr<FLoadBatch, FLoadBatchDeleter>;
 
 
 [[nodiscard]] PLAINPROPS_API FLoadBatchPtr	CreateLoadPlans(FSchemaBatchId BatchId,const FCustomBindings& Customs, const FSchemaBindings& Schemas, TConstArrayView<FStructId> RuntimeIds, ESchemaFormat Format);
+PLAINPROPS_API void							UpdateLoadPlans(FLoadBatch& Batch, const FCustomBindings& Customs, const FSchemaBindings& Schemas, TConstArrayView<FStructId> RuntimeIds);
+PLAINPROPS_API bool							HasLoadPlan(FStructSchemaId LoadId, const FLoadBatch& Batch);
 PLAINPROPS_API void							LoadStruct(void* Dst, FByteReader Src, FStructSchemaId LoadId, const FLoadBatch& Batch);
 PLAINPROPS_API void							ConstructAndLoadStruct(void* Dst, FByteReader Src, FStructSchemaId LoadId, const FLoadBatch& Batch);
 
-} // namespace PlainProps
\ No newline at end of file
+} // namespace PlainProps
diff --git a/Engine/Plugins/Experimental/PlainProps/Source/Public/PlainPropsVisualize.h b/Engine/Plugins/Experimental/PlainProps/Source/Public/PlainPropsVisualize.h
--- a/Engine/Plugins/Experimental/PlainProps/Source/Public/PlainPropsVisualize.h
+++ b/Engine/Plugins/Experimental/PlainProps/Source/Public/PlainPropsVisualize.h
@@ -73,11 +73,19 @@ struct FLoadStructPlan
 	uint64					Handle;
 };
 
+struct FLateLoadPlans
+{
+	FLateLoadPlans* Next;
+	uint8 Data[0];
+};
+
 struct FLoadBatch
 {
 	FSchemaBatchId			BatchId;
 	uint32					NumReadSchemas;
 	uint32					NumPlans;
+	uint32					MaxPlans;
+	FLateLoadPlans*			LatePlans;
 	FLoadStructPlan			Plans[0];
 };
 
diff --git a/Engine/Plugins/Experimental/PlainPropsUObject/Source/Private/PlainPropsRoundtripTest.cpp b/Engine/Plugins/Experimental/PlainPropsUObject/Source/Private/PlainPropsRoundtripTest.cpp
--- a/Engine/Plugins/Experimental/PlainPropsUObject/Source/Private/PlainPropsRoundtripTest.cpp
+++ b/Engine/Plugins/Experimental/PlainPropsUObject/Source/Private/PlainPropsRoundtripTest.cpp
@@ -485,15 +485,14 @@ private:
 class FPersistentBatchLoader
 {
 public:
-	FPersistentBatchLoader(const FCustomBindings& Customs, FMemoryView Data, TPair<FSensitiveName, FSensitiveName> RemapNames)
+	FPersistentBatchLoader(const FCustomBindings& InCustoms, FMemoryView Data, TPair<FSensitiveName, FSensitiveName> RemapNames)
+	: Customs(InCustoms)
 	{
 		// Read ids
 		FByteReader It(Data);
 		verify(It.Grab<uint32>() == Magics[0]);
 		ReadStableNames(It);
 
-// todo: re-enable remap package name
-#if 0
 		// Remap original/saved names to load/instance runtime names,
 		// for now if the package name of the batch itself has been saved (e.g. for a blueprint class scope name),
 		// then replace it with its load/instance package name
@@ -504,7 +503,6 @@ public:
 				*Name = FSensitiveName(RemapNames.Value);
 			}
 		}
-#endif
 
 		// Read and mount schemas
 		verify(It.Grab<uint32>() == Magics[1]);
@@ -533,7 +531,7 @@ public:
 		check(!Objects.IsEmpty());
 
 		// Finally create load plans
-		TConstArrayView<FStructId> LoadStructIds = RuntimeIds.Translation.GetStructIds(NumStructSchemas(Batch));
+		LoadStructIds = RuntimeIds.Translation.GetStructIds(NumStructSchemas(Batch));
 		Plans = CreateLoadPlans(Batch, Customs, GUE.Schemas, LoadStructIds, ESchemaFormat::StableNames);
 	}
 
@@ -566,20 +564,29 @@ public:
 	{
 		check(HasObject(ObjectIndex));
 		FStructView In = Objects[ObjectMap[ObjectIndex]];
+		if (!HasLoadPlan(In.Schema.Id, *Plans))
+		{
+			UpdateLoadPlans(*Plans, Customs, GUE.Schemas, LoadStructIds);
+			++NumUpdates;
+		}
 		LoadStruct(Dst, In.Values, In.Schema.Id, *Plans);
 		++NumLoaded;
 	}
 
 	static void PrintDebugInfo()
 	{
-		UE_LOGFMT(LogPlainPropsUObject, Display, "Loaded {NumLoaded} exports", NumLoaded);
+		UE_LOGFMT(LogPlainPropsUObject, Display,
+			"Loaded {NumLoaded} exports, {NumUpdates} load plan updates", NumLoaded, NumUpdates);
 	}
 
 private:
 	static inline int32			NumLoaded = 0;
+	static inline int32			NumUpdates = 0;
 
+	const FCustomBindings&		Customs;
 	const FSchemaBatch*			TranslatedSchemas = nullptr;
 	TArray<FSensitiveName>		Names;
+	TArray<FStructId>			LoadStructIds;
 	FLoadBatchPtr				Plans;
 	TArray<FStructView>			Objects;
 	TConstArrayView<uint32>		ObjectMap;
@@ -1455,6 +1462,7 @@ struct FPlainPropsLinkerSandbox
 
 	static void Shutdown()
 	{
+		MountedBaselinePackages.Empty();
 		MountedPlainpropsPackages.Empty();
 		SavedPackages.Empty();
 	}
@@ -1529,7 +1537,7 @@ struct FPlainPropsLinkerSandbox
 
 		if (PendingClasses.Num())
 		{
-			// Re-bind dynamic classes whenever their CDO has been recompiled
+			// Re-bind blueprint classes whenever their CDO has been recompiled
 			for (const UClass* Class : PendingClasses)
 			{
 				FBindId Id = GUE.Names.IndexBindId(PlainProps::UE::IndexStruct(Class));
@@ -1546,7 +1554,7 @@ struct FPlainPropsLinkerSandbox
 			}
 			PendingClasses.Reset();
 		}
-		// Bind blueprint generated classes at the time we load its CDO.
+		// Bind dynamic classes at the time we load their CDO.
 		// This is first time the struct binding is needed,
 		// and we also know for sure that the UClass has already been linked.
 		if (Object->HasAllFlags(RF_ClassDefaultObject))
@@ -1555,7 +1563,7 @@ struct FPlainPropsLinkerSandbox
 			if (!Class->HasAllClassFlags(CLASS_Native))
 			{
 				FBindId Id = GUE.Names.IndexBindId(PlainProps::UE::IndexStruct(Class));
-				UE_LOGFMT(LogPlainPropsUObject, Log, " Binding loaded blueprint class {BindId} for loading CDO {CDO}",
+				UE_LOGFMT(LogPlainPropsUObject, Log, " Binding loaded dynamic class {BindId} for loading CDO {CDO}",
 					GUE.Debug.Print(Id), *Object->GetPathName());
 				check(!GUE.Schemas.FindStruct(Id))
 				UE::BindStruct(Id, Class);
@@ -1565,20 +1573,7 @@ struct FPlainPropsLinkerSandbox
 			TRACE_CPUPROFILER_EVENT_SCOPE(LoadPlainProps);
 			BatchLoader.LoadObject(Object, ExportIndex);
 		}
-// todo: re-enable bind dynamic class
-#if 0
-		// Bind dynamic classes other than blueprints right after we have loaded them
-		if (const UClass* Class = Cast<UClass>(Object))
-		{
-			if (!Class->HasAnyClassFlags(CLASS_CompiledFromBlueprint))
-			{
-				FBindId Id = GUE.Names.IndexBindId(PlainProps::UE::IndexStruct(Class));
-				UE_LOGFMT(LogPlainPropsUObject, Display, " Binding loaded dynamic class {BindId}", GUE.Debug.Print(Id));
-				check(!GUE.Schemas.FindStruct(Id))
-				UE::BindStruct(Id, Class);
-			}
-		}
-#endif
+
 #if WITH_EDITOR
 		const int32 Pos = Ar.Tell();
 		check(Pos == SerializeContext->SerializedObjectScriptStartOffset);
@@ -2023,8 +2018,6 @@ static FName GenerateDroppedPackageName(FName Name, int32 NewNumber = 0)
 
 static void DropPackages(TArray<UPackage*>&& Packages, FDropPackagesFunc DropPackagesFunc, int32 NewNumber = INDEX_NONE)
 {
-// todo: re-enable drop packages
-#if 0
 	for (UPackage* Package : Packages)
 	{
 		ForEachObjectWithPackage(Package, [](UObject* Object)
@@ -2043,7 +2036,6 @@ static void DropPackages(TArray<UPackage*>&& Packages, FDropPackagesFunc DropPac
 	}
 
 	DropPackagesFunc(MoveTemp(Packages), NewNumber);
-#endif
 }
 
 int32 RoundtripViaLinkerBatch(
@@ -2088,22 +2080,6 @@ int32 RoundtripViaLinkerBatch(
 			{
 				continue;
 			}
-			// todo: re-enable blueprints
-			bool bExclude = false;
-			ForEachObjectWithPackage(Package, [&bExclude](UObject* Object)
-			{
-				UClass* Class = Object->GetClass();
-				if (Class->HasAnyClassFlags(CLASS_CompiledFromBlueprint))
-				{
-					bExclude = true;
-					return false;
-				}
-				return true;
-			});
-			if (bExclude)
-			{
-				continue;
-			}
 			OriginalPackages.Add(Package);
 			PackageNames.Add(Package->GetFName());
 		}
@@ -2289,14 +2265,24 @@ int32 RoundtripViaLinkerBatch(
 		}
 	}
 
-	TRACE_CPUPROFILER_EVENT_FLUSH();
-	TRACE_CPUPROFILER_EVENT_SCOPE(Shutdown);
 	FPersistentBatchLoader::PrintDebugInfo();
+	TRACE_CPUPROFILER_EVENT_FLUSH();
 	
+	{
+		TRACE_CPUPROFILER_EVENT_SCOPE(CollectGarbage);
+		UE_LOGFMT(LogPlainPropsUObject, Log, "Collecting garbage...");
+		CollectGarbage(GARBAGE_COLLECTION_KEEPFLAGS);
+	}
+
 #if UE_BUILD_DEBUG
 	// Note: Shutdown is too slow to be enabled by default.
-	// It will reset and flush all the linkers which takes ~15s for ~3000 packages in Release config
-	FPlainPropsLinkerSandbox::Shutdown();
+	// It will reset and flush all the linkers which loads and detaches all bulk data from the saved memory archive.
+	// This takes ~15s for ~3000 packages in Release config.
+	{
+		TRACE_CPUPROFILER_EVENT_SCOPE(Shutdown);
+		FScopeLoadTestPackage LoadScope(/*bPlainProps*/false);
+		FPlainPropsLinkerSandbox::Shutdown();
+	}
 #endif
 
 #endif // PP_WITH_ENGINE_HOOKS
