nss: Fixed truncated grab-token error string due to tchar/char8_t confusion.

#rnx

#Synced-CL 49600073

diff --git a/Engine/Plugins/Experimental/PlainProps/Source/Private/PlainPropsParse.cpp b/Engine/Plugins/Experimental/PlainProps/Source/Private/PlainPropsParse.cpp
--- a/Engine/Plugins/Experimental/PlainProps/Source/Private/PlainPropsParse.cpp
+++ b/Engine/Plugins/Experimental/PlainProps/Source/Private/PlainPropsParse.cpp
@@ -958,7 +958,11 @@ void FYamlTokenizer::SetLastError(FUtf8StringView Message)
 	LastError.Reset();
 	if (Message.Len())
 	{
-		LastError.Appendf("'%s' in line %u: '%s'", *Print(Message), ReadLine.LineNumber, *Print(ReadLine.FullView));
+		LastError.Appendf("'%.*s' in line %u: '%.*s'",
+			Message.Len(), Message.GetData(),
+			ReadLine.LineNumber,
+			ReadLine.FullView.Len(),
+			ReadLine.FullView.GetData());
 	}
 }
 
