nss: Roundtrip enums as strings instead of integers.

#rnx

#Synced-CL 49920676

diff --git a/Engine/Plugins/Experimental/PlainProps/Source/Private/PlainPropsInternalRead.h b/Engine/Plugins/Experimental/PlainProps/Source/Private/PlainPropsInternalRead.h
--- a/Engine/Plugins/Experimental/PlainProps/Source/Private/PlainPropsInternalRead.h
+++ b/Engine/Plugins/Experimental/PlainProps/Source/Private/PlainPropsInternalRead.h
@@ -15,6 +15,11 @@ inline const uint8*			GetValues(FLeafRangeView Range)		{ return Range.Values; }
 inline FRangeSchema			GetSchema(FRangeView Range)			{ return Range.Schema; }
 inline FMemoryView			GetValues(FRangeView Range)			{ return Range.Values; }
 
+inline FOptionalEnumSchemaId GetEnumSchemaId(FLeafRangeView Range)
+{
+	return Range.Enum;
+}
+
 template<typename MemberType, typename IndexType>
 TConstArrayView<MemberType> GrabInnerRangeTypes(TConstArrayView<MemberType> InnerRangeTypes, IndexType& InOutIdx)
 {
diff --git a/Engine/Plugins/Experimental/PlainProps/Source/Private/PlainPropsInternalText.h b/Engine/Plugins/Experimental/PlainProps/Source/Private/PlainPropsInternalText.h
--- a/Engine/Plugins/Experimental/PlainProps/Source/Private/PlainPropsInternalText.h
+++ b/Engine/Plugins/Experimental/PlainProps/Source/Private/PlainPropsInternalText.h
@@ -27,6 +27,8 @@ struct FLiterals
 	FAnsiStringView Super = "__super__";
 	FAnsiStringView Dynamic = "__dynamic__";
 
+	FAnsiStringView UnsetFlagEnum = "!UnsetFlagEnum"; // '!' so literal is not a valid identifier
+
 	FAnsiStringView Oob = "ERR_oob";
 
 	FAnsiStringView Ranges[9] = { 
diff --git a/Engine/Plugins/Experimental/PlainProps/Source/Private/PlainPropsParse.cpp b/Engine/Plugins/Experimental/PlainProps/Source/Private/PlainPropsParse.cpp
--- a/Engine/Plugins/Experimental/PlainProps/Source/Private/PlainPropsParse.cpp
+++ b/Engine/Plugins/Experimental/PlainProps/Source/Private/PlainPropsParse.cpp
@@ -456,6 +456,50 @@ FType MakeType(const FTypeTokens& TypeTokens, FTextIndexer& Names)
 
 ////////////////////////////////////////////////////////////////////////////////
 
+template <UnsignedIntegral T>
+static void ParseEnumValue(
+	FUtf8StringView String,
+	T& OutValue,
+	const FParsedEnumSchema& Schema,
+	const FTextIndexer& Names)
+{
+	OutValue = 0;
+
+	if (String == GLiterals.UnsetFlagEnum)
+	{
+		return;
+	}
+
+	while (!String.IsEmpty())
+	{
+		static constexpr FAsciiSet Separators(" |");
+
+		String = FAsciiSet::TrimPrefixWith(String, Separators);
+		FUtf8StringView Element = FAsciiSet::FindPrefixWithout(String, Separators);
+		check(!Element.IsEmpty());
+
+		String = String.Mid(Element.Len());
+
+		bool bEnumeratorFound = false;
+		for (const auto& [Name, Constant] : Schema.Enumerators)
+		{
+			FUtf8StringView NameView = Names.ResolveName(Name);
+			if (NameView != Element)
+			{
+				continue;
+			}
+
+			OutValue |= Constant;
+			bEnumeratorFound = true;
+			break;
+		}
+
+		checkf(bEnumeratorFound, TEXT("Unable to match enumerator '%.*hs'"), String.Len(), String.GetData());
+	}
+}
+
+////////////////////////////////////////////////////////////////////////////////
+
 struct FLineInfo
 {
 	FUtf8StringView			View;
@@ -1398,12 +1442,11 @@ void AddMember(FMemberBuilder& Members, FMemberId Id, FUtf8StringView String)
 }
 
 template<UnsignedIntegral T>
-void AddEnum(FMemberBuilder& Members, FMemberId Id, FEnumId Enum, FUtf8StringView String)
+void AddEnum(FMemberBuilder& Members, FMemberId Id, const FParsedEnumSchema& Schema, const FTextIndexer& Names, FUtf8StringView String)
 {
-	if (TOptional<T> Value = Parse<T>(String))
-	{
-		Members.AddEnum(Id, Enum, Value.GetValue());
-	}
+	T Value;
+	ParseEnumValue(String, Value, Schema, Names);
+	Members.AddEnum(Id, Schema.Id, Value);
 }
 
 template<UnsignedIntegral T>
@@ -1491,10 +1534,10 @@ bool FMemberParser::ParseLeaf()
 		case ELeafType::Enum:
 			switch (Leaf.Width)
 			{
-				case ELeafWidth::B8:	AddEnum<uint8 >(MemberBuilder, MemberId, Enum.Get(), String); break;
-				case ELeafWidth::B16:	AddEnum<uint16>(MemberBuilder, MemberId, Enum.Get(), String); break;
-				case ELeafWidth::B32:	AddEnum<uint32>(MemberBuilder, MemberId, Enum.Get(), String); break;
-				case ELeafWidth::B64:	AddEnum<uint64>(MemberBuilder, MemberId, Enum.Get(), String); break;
+				case ELeafWidth::B8:	AddEnum<uint8 >(MemberBuilder, MemberId, ResolveEnumSchema(Schemas, Enum.Get()), Names, String); break;
+				case ELeafWidth::B16:	AddEnum<uint16>(MemberBuilder, MemberId, ResolveEnumSchema(Schemas, Enum.Get()), Names, String); break;
+				case ELeafWidth::B32:	AddEnum<uint32>(MemberBuilder, MemberId, ResolveEnumSchema(Schemas, Enum.Get()), Names, String); break;
+				case ELeafWidth::B64:	AddEnum<uint64>(MemberBuilder, MemberId, ResolveEnumSchema(Schemas, Enum.Get()), Names, String); break;
 			}
 			break;
 		case ELeafType::Unicode:
@@ -1652,7 +1695,8 @@ template<Arithmetic T>
 template<UnsignedIntegral T>
 [[nodiscard]] static FTypedRange ParseEnumRange(
 	FScratchAllocator& Scratch,
-	FEnumId Enum,
+	const FParsedEnumSchema& Schema,
+	const FTextIndexer& Names,
 	ESizeType MaxSize,
 	FScopedTokenReader& RangeTokens)
 {
@@ -1660,12 +1704,12 @@ template<UnsignedIntegral T>
 	while (RangeTokens.HasMore())
 	{
 		FToken Token = RangeTokens.GrabToken();
-		if (TOptional<T> Value = Parse<T>(Token.Value()))
-		{
-			Values.Add(Value.GetValue());
-		}
+		FUtf8StringView String = Token.Value();
+		T Value;
+		ParseEnumValue(String, Value, Schema, Names);
+		Values.Add(Value);
 	}
-	return BuildEnumRange(Scratch, Enum, MaxSize, TConstArrayView<T>(Values));
+	return BuildEnumRange(Scratch, Schema.Id, MaxSize, TConstArrayView<T>(Values));
 }
 
 template<UnsignedIntegral T>
@@ -1740,10 +1784,10 @@ FTypedRange FMemberParser::ParseLeaves(const FMemberSchema& RangeSchema, FScoped
 		case ELeafType::Enum:
 			switch (Leaf.Width)
 			{
-				case ELeafWidth::B8:	TypedRange = ParseEnumRange<uint8 >(Scratch, Enum.Get(), NumType, RangeTokens); break;
-				case ELeafWidth::B16:	TypedRange = ParseEnumRange<uint16>(Scratch, Enum.Get(), NumType, RangeTokens); break;
-				case ELeafWidth::B32:	TypedRange = ParseEnumRange<uint32>(Scratch, Enum.Get(), NumType, RangeTokens); break;
-				case ELeafWidth::B64:	TypedRange = ParseEnumRange<uint64>(Scratch, Enum.Get(), NumType, RangeTokens); break;
+				case ELeafWidth::B8:	TypedRange = ParseEnumRange<uint8 >(Scratch, ResolveEnumSchema(Schemas, Enum.Get()), Names, NumType, RangeTokens); break;
+				case ELeafWidth::B16:	TypedRange = ParseEnumRange<uint16>(Scratch, ResolveEnumSchema(Schemas, Enum.Get()), Names, NumType, RangeTokens); break;
+				case ELeafWidth::B32:	TypedRange = ParseEnumRange<uint32>(Scratch, ResolveEnumSchema(Schemas, Enum.Get()), Names, NumType, RangeTokens); break;
+				case ELeafWidth::B64:	TypedRange = ParseEnumRange<uint64>(Scratch, ResolveEnumSchema(Schemas, Enum.Get()), Names, NumType, RangeTokens); break;
 			}
 			break;
 		case ELeafType::Unicode:
diff --git a/Engine/Plugins/Experimental/PlainProps/Source/Private/PlainPropsPrint.cpp b/Engine/Plugins/Experimental/PlainProps/Source/Private/PlainPropsPrint.cpp
--- a/Engine/Plugins/Experimental/PlainProps/Source/Private/PlainPropsPrint.cpp
+++ b/Engine/Plugins/Experimental/PlainProps/Source/Private/PlainPropsPrint.cpp
@@ -257,6 +257,77 @@ void PrintMemberSchema(FUtf8Builder& Out, const FIds& Ids, FMemberSchema Member)
 	}	
 }
 
+///////////////////////////////////////////////////////////////////////////////
+template <UnsignedIntegral T, bool bFlagIndexing>
+static FNameId GetEnumConstantId(const FEnumSchema& Schema, T Value)
+{
+	using ValueType = decltype(Value);
+
+	TConstArrayView<FNameId> Names = MakeConstArrayView(Schema.Footer, Schema.Num);
+
+	if (Schema.ExplicitConstants)
+	{
+		for (uint32 Idx = 0; ValueType Constant : GetConstants<ValueType>(Schema))
+		{
+			if (Constant == Value)
+			{
+				return Names[Idx];
+			}
+			++Idx;
+		}
+	}
+	else
+	{
+		if constexpr (bFlagIndexing)
+		{
+			Value = 63 - FMath::CountLeadingZeros64(uint64(Value));
+		}
+
+		if (Value < ValueType(Names.Num()))
+		{
+			return Names[Value];
+		}
+	}
+
+	return {};
+}
+
+template <UnsignedIntegral T>
+static void PrintEnumValue(
+	FUtf8StringBuilderBase& Buffer,
+	T Value,
+	const FEnumSchema& Schema,
+	const FBatchIds& Ids)
+{
+	if (Schema.FlagMode == 0)
+	{
+		FNameId ConstantId = GetEnumConstantId<T, false>(Schema, Value);
+		checkf(ConstantId != FNameId(), TEXT("Unknown enum constant value '%lld'"), int64(Value));
+		Ids.AppendString(Buffer, ConstantId);
+		return;
+	}
+
+	if (Value == 0)
+	{
+		Buffer << GLiterals.UnsetFlagEnum;
+		return;
+	}
+
+	using ValueType = decltype(Value);
+	FUtf8StringView Separator;
+	do
+	{
+		ValueType Flag = Value & (~Value + 1);
+		Value ^= Flag;
+		FNameId ConstantId = GetEnumConstantId<T, true>(Schema, Flag);
+		checkf(ConstantId != FNameId(), TEXT("Unknown enum flag value '%llx'"), uint64(Flag));
+		Buffer << Separator;
+		Ids.AppendString(Buffer, ConstantId);
+		Separator = " | ";
+	}
+	while (Value != 0);
+}
+
 ///////////////////////////////////////////////////////////////////////////////
 
 struct FMemberSchemaView
@@ -408,10 +479,14 @@ public:
 
 private:
 	void PrintLeaf(FMemberId Id, FLeafView LeafView);
+	template <UnsignedIntegral T>
+	void PrintLeafEnum(FMemberId Id, FLeafView LeafView);
 	void PrintStruct(FOptionalMemberId Id, FStructType StructType, FStructView StructView);
 	void PrintRange(FMemberId Id, FRangeType RangeType, const FRangeView& RangeView);
 
 	void PrintLeaves(FUnpackedLeafType Leaf, const FLeafRangeView& LeafRange);
+	template <UnsignedIntegral T>
+	void PrintLeavesEnum(FUnpackedLeafType Leaf, const FLeafRangeView& LeafRange);
 	void PrintStructs(FStructType StructType, const FStructRangeView& StructRange);
 	void PrintRanges(FRangeType RangeType, const FNestedRangeView& NestedRange);
 
@@ -1009,10 +1084,10 @@ void FMemberPrinter::PrintLeaf(FMemberId Id, FLeafView LeafView)
 	case ELeafType::Enum:
 		switch (LeafView.Leaf.Width)
 		{
-			case ELeafWidth::B8:	TextBuilder.AddValue(LeafView.AsUnderlyingValue<uint8>());  break;
-			case ELeafWidth::B16:	TextBuilder.AddValue(LeafView.AsUnderlyingValue<uint16>()); break;
-			case ELeafWidth::B32:	TextBuilder.AddValue(LeafView.AsUnderlyingValue<uint32>()); break;
-			case ELeafWidth::B64:	TextBuilder.AddValue(LeafView.AsUnderlyingValue<uint64>()); break;
+			case ELeafWidth::B8:	PrintLeafEnum<uint8 >(Id, LeafView); break;
+			case ELeafWidth::B16:	PrintLeafEnum<uint16>(Id, LeafView); break;
+			case ELeafWidth::B32:	PrintLeafEnum<uint32>(Id, LeafView); break;
+			case ELeafWidth::B64:	PrintLeafEnum<uint64>(Id, LeafView); break;
 		}
 		break;
 	case ELeafType::Unicode:
@@ -1021,13 +1096,30 @@ void FMemberPrinter::PrintLeaf(FMemberId Id, FLeafView LeafView)
 			case ELeafWidth::B8:	TextBuilder.AddValue(LeafView.AsChar8());	break;
 			case ELeafWidth::B16:	TextBuilder.AddValue(LeafView.AsChar16());	break;
 			case ELeafWidth::B32:	TextBuilder.AddValue(LeafView.AsChar32());	break;
-			case ELeafWidth::B64:	check(false);									break;
+			case ELeafWidth::B64:	check(false);								break;
 		};
 	}
 
 	PrintSchemaComment(LeafView.Leaf, LeafView.Enum);
 }
 
+template <UnsignedIntegral T>
+void FMemberPrinter::PrintLeafEnum(FMemberId Id, FLeafView LeafView)
+{
+	check(LeafView.Leaf.Type == ELeafType::Enum);
+
+	T Value = LeafView.AsUnderlyingValue<T>();
+
+	FEnumSchemaId EnumSchemaId = LeafView.Enum.Get();
+	const FSchemaBatch& SchemaBatch = Ids.GetSchemas();
+	const FEnumSchema& EnumSchema = ResolveEnumSchema(SchemaBatch, EnumSchemaId);
+
+	TUtf8StringBuilder<64> Buffer;
+	PrintEnumValue(Buffer, Value, EnumSchema, Ids);
+
+	TextBuilder.AddValue(FUtf8StringView(Buffer));
+}
+
 void FMemberPrinter::PrintStruct(FOptionalMemberId MemberId, FStructType StructType, FStructView StructView)
 {
 	TextBuilder.BeginStruct(*FPrintId<128>(Ids, MemberId));
@@ -1094,10 +1186,10 @@ void FMemberPrinter::PrintLeaves(FUnpackedLeafType Leaf, const FLeafRangeView& L
 	case ELeafType::Enum:
 		switch (Leaf.Width)
 		{
-			case ELeafWidth::B8:  for (const uint8&  U : LeafRange.AsUnderlyingValues<uint8>())  { TextBuilder.AddItem(U); } break;
-			case ELeafWidth::B16: for (const uint16& U : LeafRange.AsUnderlyingValues<uint16>()) { TextBuilder.AddItem(U); } break;
-			case ELeafWidth::B32: for (const uint32& U : LeafRange.AsUnderlyingValues<uint32>()) { TextBuilder.AddItem(U); } break;
-			case ELeafWidth::B64: for (const uint64& U : LeafRange.AsUnderlyingValues<uint64>()) { TextBuilder.AddItem(U); } break;
+			case ELeafWidth::B8:  PrintLeavesEnum<uint8 >(Leaf, LeafRange); break;
+			case ELeafWidth::B16: PrintLeavesEnum<uint16>(Leaf, LeafRange); break;
+			case ELeafWidth::B32: PrintLeavesEnum<uint32>(Leaf, LeafRange); break;
+			case ELeafWidth::B64: PrintLeavesEnum<uint64>(Leaf, LeafRange); break;
 		}
 		break;
 	case ELeafType::Unicode:
@@ -1106,6 +1198,21 @@ void FMemberPrinter::PrintLeaves(FUnpackedLeafType Leaf, const FLeafRangeView& L
 	}
 }
 
+template <UnsignedIntegral T>
+void FMemberPrinter::PrintLeavesEnum(FUnpackedLeafType Leaf, const FLeafRangeView& LeafRange)
+{
+	FEnumSchemaId EnumSchemaId = GetEnumSchemaId(LeafRange).Get();
+	const FSchemaBatch& SchemaBatch = Ids.GetSchemas();
+	const FEnumSchema& EnumSchema = ResolveEnumSchema(SchemaBatch, EnumSchemaId);
+
+	for (const T Value: LeafRange.AsUnderlyingValues<T>())
+	{
+		TUtf8StringBuilder<128> Buffer;
+		PrintEnumValue(Buffer, Value, EnumSchema, Ids);
+		TextBuilder.AddItem(FUtf8StringView(Buffer));
+	}
+}
+
 void FMemberPrinter::PrintStructs(FStructType StructType, const FStructRangeView& StructRange)
 {
 	for (FStructView StructView : StructRange)
diff --git a/Engine/Plugins/Experimental/PlainProps/Source/Public/PlainPropsRead.h b/Engine/Plugins/Experimental/PlainProps/Source/Public/PlainPropsRead.h
--- a/Engine/Plugins/Experimental/PlainProps/Source/Public/PlainPropsRead.h
+++ b/Engine/Plugins/Experimental/PlainProps/Source/Public/PlainPropsRead.h
@@ -538,6 +538,7 @@ private:
 
 	friend FUnpackedLeafType GetType(FLeafRangeView Range);
 	friend const uint8* GetValues(FLeafRangeView Range);
+	friend FOptionalEnumSchemaId GetEnumSchemaId(FLeafRangeView Range);
 };
 
 //////////////////////////////////////////////////////////////////////////
