PlainProps fixes

* Fix inner range type for statically bound structs
* Fix loose indexer ptr stability
* Fix FSets unit test impacted by new compact TSet
* Tiny natvis tweak

#rnx
#Synced-CL 49944300

diff --git a/Engine/Plugins/Experimental/PlainProps/Source/Private/PlainPropsBind.cpp b/Engine/Plugins/Experimental/PlainProps/Source/Private/PlainPropsBind.cpp
--- a/Engine/Plugins/Experimental/PlainProps/Source/Private/PlainPropsBind.cpp
+++ b/Engine/Plugins/Experimental/PlainProps/Source/Private/PlainPropsBind.cpp
@@ -504,14 +504,11 @@ static void DeclareOrValidate(FStructDeclarationPtr& Out, FStructSpec Info, TCon
 	FMemory::Memcpy(Ptr->MemberNames, Info.MemberNames.GetData(), Info.MemberNames.NumBytes());
 	FMemory::Memcpy(const_cast<FInnerId*>(Ptr->GetInnerIds()), InnerIds.GetData(), InnerIds.NumBytes());
 	FMemberType* TypeIt = const_cast<FMemberType*>(Ptr->GetTypes());
-	for (FMemberBindType BindType : MakeArrayView(Binding.Members, Binding.NumMembers))
+	checkSlow(TypeIt + Num == Ptr->GetInnerRangeTypes() && Binding.Members + Num == Binding.GetInnerRangeTypes()); // Loop assumes inner range types follow member types
+	for (FMemberBindType BindType : MakeArrayView(Binding.Members, Binding.NumMembers + Binding.NumInnerRanges))
 	{
 		*TypeIt++ = ToMemberType(BindType);
 	}
-	for (FRangeBinding Range : MakeArrayView(Binding.GetRangeBindings(), Binding.NumInnerRanges))
-	{
-		*TypeIt++ = FMemberType(Range.GetSizeType());
-	}
 
 	// Make smart pointer
 	new (&Out) FStructDeclarationPtr(Ptr, /* add ref */ false);
diff --git a/Engine/Plugins/Experimental/PlainProps/Source/Private/PlainPropsInternalUpgrade.h b/Engine/Plugins/Experimental/PlainProps/Source/Private/PlainPropsInternalUpgrade.h
--- a/Engine/Plugins/Experimental/PlainProps/Source/Private/PlainPropsInternalUpgrade.h
+++ b/Engine/Plugins/Experimental/PlainProps/Source/Private/PlainPropsInternalUpgrade.h
@@ -58,9 +58,9 @@ public:
 //	void						Drop(FDeclId Id)										{ Structs.Remove(Id); }
 
 private:
-	TMap<FDeclId, FParameterIndexer>	Structs;
-	TMap<FEnumId, FEnumeratorIndexer>	Enums;
-	const IDeclarations&				Declarations;
+	TMap<FDeclId, TUniquePtr<FParameterIndexer>>	Structs;
+	TMap<FEnumId, TUniquePtr<FEnumeratorIndexer>>	Enums;
+	const IDeclarations&							Declarations;
 };
 
 ////////////////////////////////////////////////////////////////////////////////////////////////
diff --git a/Engine/Plugins/Experimental/PlainProps/Source/Private/PlainPropsUpgrade.cpp b/Engine/Plugins/Experimental/PlainProps/Source/Private/PlainPropsUpgrade.cpp
--- a/Engine/Plugins/Experimental/PlainProps/Source/Private/PlainPropsUpgrade.cpp
+++ b/Engine/Plugins/Experimental/PlainProps/Source/Private/PlainPropsUpgrade.cpp
@@ -91,17 +91,15 @@ static FParameterIdSet GetDeclaredOrder(const FStructDeclaration& Decl)
 }
 
 template<class IndexerT, typename IdT>
-IndexerT& GetIndexer(TMap<IdT, IndexerT>& Map, IdT Id, const IDeclarations& Declarations)
+IndexerT& GetIndexer(TMap<IdT, TUniquePtr<IndexerT>>& Map, IdT Id, const IDeclarations& Declarations)
 {
-	if (IndexerT* Out = Map.Find(Id))
+	TUniquePtr<IndexerT>& Out = Map.FindOrAdd(Id);
+	if (!Out)
 	{
-		return *Out;
+		const auto Decl = Declarations.Find(Id);
+		Out.Reset(Decl ? new IndexerT(GetDeclaredOrder(*Decl)) : new IndexerT);
 	}
-	else if (const auto Decl = Declarations.Find(Id))
-	{
-		return Map.Emplace(Id, GetDeclaredOrder(*Decl));
-	}
-	return Map.Emplace(Id);
+	return *Out;
 }
 
 FEnumeratorIndexer& FLooseIndexers::Get(FEnumId Id) { return GetIndexer(/* in-out */ Enums, Id, Declarations); }
diff --git a/Engine/Plugins/Experimental/PlainProps/Source/Tests/PlainPropsUeCoreTests.cpp b/Engine/Plugins/Experimental/PlainProps/Source/Tests/PlainPropsUeCoreTests.cpp
--- a/Engine/Plugins/Experimental/PlainProps/Source/Tests/PlainPropsUeCoreTests.cpp
+++ b/Engine/Plugins/Experimental/PlainProps/Source/Tests/PlainPropsUeCoreTests.cpp
@@ -1187,7 +1187,7 @@ TEST_CASE_NAMED(FPlainPropsUeCoreTest, "System::Core::Serialization::PlainProps:
 				CHECK(!Same<ECompare::Order>(FSets{{'a','b'}}, FSets{{'b','a'}}));
 				CHECK(Same<ECompare::Order>(Batch.Load<FSets>(), FSets{{'a','b'}}));
 				CHECK(Same<ECompare::Order>(Batch.Load<FSets>(), FSets{{'b','a'}}));
-				CHECK(Same<ECompare::Order>(Batch.Load<FSets>(), FSets{{'z','a','?'}}));
+				CHECK(Same<ECompare::Content>(Batch.Load<FSets>(), FSets{{'z','a','?'}}));
 			});
 	}
 
