nss: Use literals for errors that aren't valid identifiers

#rnx

#Synced-CL 49920676

diff --git a/Engine/Plugins/Experimental/PlainProps/Source/Private/PlainPropsInternalText.h b/Engine/Plugins/Experimental/PlainProps/Source/Private/PlainPropsInternalText.h
--- a/Engine/Plugins/Experimental/PlainProps/Source/Private/PlainPropsInternalText.h
+++ b/Engine/Plugins/Experimental/PlainProps/Source/Private/PlainPropsInternalText.h
@@ -10,6 +10,8 @@ namespace PlainProps
 
 struct FLiterals
 {
+	// '!' is used in some literals so that they are not valid identifiers
+
 	FAnsiStringView Structs = "Structs";
 	FAnsiStringView Enums = "Enums";
 	FAnsiStringView Objects = "Objects";
@@ -27,21 +29,21 @@ struct FLiterals
 	FAnsiStringView Super = "__super__";
 	FAnsiStringView Dynamic = "__dynamic__";
 
-	FAnsiStringView UnsetFlagEnum = "!UnsetFlagEnum"; // '!' so literal is not a valid identifier
+	FAnsiStringView UnsetFlagEnum = "!UnsetFlagEnum";
 
-	FAnsiStringView Oob = "ERR_oob";
+	FAnsiStringView Oob = "ERR!oob";
 
 	FAnsiStringView Ranges[9] = { 
 		("Uni"), ("i8"), ("u8"), ("i16"), ("u16"), ("i32"), ("u32"), ("i64"), ("u64") };
 
 	FAnsiStringView Leaves[8][4] = {
-		{("bool"),		("ERR_b16"),	("ERR_b32"),	("ERR_b64")},
+		{("bool"),		("ERR!b16"),	("ERR!b32"),	("ERR!b64")},
 		{("i8"),		("i16"),		("i32"),		("i64")},
 		{("u8"),		("u16"),		("u32"),		("u64")},
-		{("ERR_fp8"),	("ERR_fp16"),	("float"),		("double")},
+		{("ERR!fp8"),	("ERR!fp16"),	("float"),		("double")},
 		{("hex8"),		("hex16"),		("hex32"),		("hex64")},
 		{("enum8"),		("enum16"),		("enum32"),		("enum64")},
-		{("utf8"),		("utf16"),		("utf32"),		("ERR_utf64")},
+		{("utf8"),		("utf16"),		("utf32"),		("ERR!utf64")},
 		{Oob,			Oob,			Oob,			Oob}};
 
 	FAnsiStringView Widths[4] = { ("u8"), ("u16"), ("u32"), ("u64") };
