PlainProps another fix due to TCompactSet vs TSparseSet semantics

#rnx

#Synced-CL 49944300

diff --git a/Engine/Plugins/Experimental/PlainProps/Source/Private/PlainPropsBind.cpp b/Engine/Plugins/Experimental/PlainProps/Source/Private/PlainPropsBind.cpp
--- a/Engine/Plugins/Experimental/PlainProps/Source/Private/PlainPropsBind.cpp
+++ b/Engine/Plugins/Experimental/PlainProps/Source/Private/PlainPropsBind.cpp
@@ -218,7 +218,7 @@ FCustomBindingMap::FCustomBindingMap(FDebugIds Dbg)
 
 FCustomBindingMap::~FCustomBindingMap()
 {
-	for (TSet<FBindId>::TConstIterator It(Keys); It; ++It)
+	for (TSparseSet<FBindId>::TConstIterator It(Keys); It; ++It)
 	{
 		int32 Idx = It.GetId().AsInteger();
 		FCustomBindingHandle Handle = {Values + Idx, MaxValues};
diff --git a/Engine/Plugins/Experimental/PlainProps/Source/Private/PlainPropsInternalUpgrade.h b/Engine/Plugins/Experimental/PlainProps/Source/Private/PlainPropsInternalUpgrade.h
--- a/Engine/Plugins/Experimental/PlainProps/Source/Private/PlainPropsInternalUpgrade.h
+++ b/Engine/Plugins/Experimental/PlainProps/Source/Private/PlainPropsInternalUpgrade.h
@@ -329,9 +329,9 @@ struct FTransformation
 struct FTransformations
 {
 	TArray<FTransformation>		Structs;
-	TSet<FDeclId>				StructLookup;
+	TSet<FDeclId>				StructLookup;		// If we start removing, switch to TSparseSet for FSetElementId stability
 	TArray<FEnumTransformation>	Enums;
-	TSet<FEnumId>				EnumLookup;
+	TSet<FEnumId>				EnumLookup;			// If we start removing, switch to TSparseSet for FSetElementId stability
 	int32						NumSourceStructs;
 
 	const FTransformation&		operator[](FStructSchemaId Id) const	{ return Structs[Id.Idx]; }
diff --git a/Engine/Plugins/Experimental/PlainProps/Source/Public/PlainPropsBind.h b/Engine/Plugins/Experimental/PlainProps/Source/Public/PlainPropsBind.h
--- a/Engine/Plugins/Experimental/PlainProps/Source/Public/PlainPropsBind.h
+++ b/Engine/Plugins/Experimental/PlainProps/Source/Public/PlainPropsBind.h
@@ -3,7 +3,7 @@
 #pragma once
 
 #include "Containers/ArrayView.h"
-#include "Containers/Set.h"
+#include "Containers/SparseSet.h"
 #include "Memory/MemoryFwd.h"
 #include "Memory/MemoryView.h"
 #include "PlainPropsDeclare.h"
@@ -290,7 +290,7 @@ struct FCustomBindingMap
 	void								Drop(FBindId Id);
 
 	FIdWindow							Window;
-	TSet<FBindId>						Keys;
+	TSparseSet<FBindId>					Keys; // Need FSetElementId stability when dropping
 	ICustomBinding**					Values = nullptr;
 	uint32								MaxValues = 0;
 	TArray<FStructDeclarationPtr>		Declarations; // Consider switch to TSparseArray or TMap
