nss: Balanced text-serialisation markup, key-value pairs need to be enclosed in structs.

#rnx

#Synced-CL 49920676

diff --git a/Engine/Plugins/Experimental/PlainProps/Source/Private/PlainPropsParse.cpp b/Engine/Plugins/Experimental/PlainProps/Source/Private/PlainPropsParse.cpp
--- a/Engine/Plugins/Experimental/PlainProps/Source/Private/PlainPropsParse.cpp
+++ b/Engine/Plugins/Experimental/PlainProps/Source/Private/PlainPropsParse.cpp
@@ -1076,7 +1076,7 @@ bool FEnumSchemaParser::ParseEnumSchema(FParsedEnumSchema& Out)
 				(void)ScopedReader.GrabToken();
 			}
 		}
-		else if (Token.Token == EToken::BeginRange)
+		else if (Token.Token == EToken::BeginStruct)
 		{
 			checkf(Token.Value() == GLiterals.Constants,
 				TEXT("Invalid range '%s' for enum"), *Print(Token));
@@ -1186,7 +1186,7 @@ bool FStructSchemaParser::ParseStructSchema(FParsedStructSchema& Out)
 				(void)ScopedReader.GrabToken();
 			}
 		}
-		else if (Token.Token == EToken::BeginRange)
+		else if (Token.Token == EToken::BeginStruct)
 		{
 			checkf(Token.Value() == GLiterals.Members,
 				TEXT("Invalid range '%s' for struct"), *Print(Token));
@@ -2022,7 +2022,7 @@ void FBatchParser::Tokenize()
 			continue;
 		}
 
-		if (Token->Depth == 1 && Token->Token == EToken::BeginRange)
+		if (Token->Depth == 1 && Token->Token == EToken::BeginStruct)
 		{
 			if (Token->Value() == GLiterals.Structs)
 			{
@@ -2032,7 +2032,11 @@ void FBatchParser::Tokenize()
 			{
 				EnumsIdx = Tokens.Num();
 			}
-			else if (Token->Value() == GLiterals.Objects)
+		}
+
+		else if (Token->Depth == 1 && Token->Token == EToken::BeginRange)
+		{
+			if (Token->Value() == GLiterals.Objects)
 			{
 				ObjectsIdx = Tokens.Num();
 			}
diff --git a/Engine/Plugins/Experimental/PlainProps/Source/Private/PlainPropsPrint.cpp b/Engine/Plugins/Experimental/PlainProps/Source/Private/PlainPropsPrint.cpp
--- a/Engine/Plugins/Experimental/PlainProps/Source/Private/PlainPropsPrint.cpp
+++ b/Engine/Plugins/Experimental/PlainProps/Source/Private/PlainPropsPrint.cpp
@@ -746,6 +746,8 @@ void FYamlBuilder::EndRange()
 
 void FYamlBuilder::AddId(FUtf8StringView Id)
 {
+	check(!Stack.IsEmpty() && Stack.Last().IsInStruct);
+
 	AppendNewLine();
 	AppendIndentation();
 	AppendIdentifier(Id);
@@ -756,6 +758,8 @@ void FYamlBuilder::AddId(FUtf8StringView Id)
 template<typename T>
 void FYamlBuilder::AddValue(T Value)
 {
+	check(!Stack.IsEmpty() && Stack.Last().IsInStruct);
+
 	Text.AppendChar(' ');
 	AppendValue(Value);
 	IsNewLine = false;
@@ -765,6 +769,8 @@ void FYamlBuilder::AddValue(T Value)
 template<UnsignedIntegral T>
 void FYamlBuilder::AddValueHex(T Value)
 {
+	check(!Stack.IsEmpty() && Stack.Last().IsInStruct);
+
 	Text.AppendChar(' ');
 	AppendValueHex(Value);
 	IsNewLine = false;
@@ -774,6 +780,8 @@ void FYamlBuilder::AddValueHex(T Value)
 template<typename T>
 void FYamlBuilder::AddItem(T Value)
 {
+	check(!Stack.IsEmpty() && Stack.Last().IsInStruct == false);
+
 	AppendNewLine();
 	AppendIndentation();
 	AppendValue(Value);
@@ -885,19 +893,19 @@ FBatchPrinter::~FBatchPrinter()
 
 void FBatchPrinter::PrintSchemas()
 {
-	TextBuilder.BeginRange(GLiterals.Structs);
+	TextBuilder.BeginStruct(GLiterals.Structs);
 	for (const FStructSchema& Struct : GetStructSchemas(Ids.GetSchemas()))
 	{
 		PrintStructSchema(Struct, Ids.GetBatchId());
 	}
-	TextBuilder.EndRange();
+	TextBuilder.EndStruct();
 
-	TextBuilder.BeginRange(GLiterals.Enums);
+	TextBuilder.BeginStruct(GLiterals.Enums);
 	for (const FEnumSchema& EnumSchema : GetEnumSchemas(Ids.GetSchemas()))
 	{
 		PrintEnumSchema(EnumSchema);
 	}
-	TextBuilder.EndRange();
+	TextBuilder.EndStruct();
 }
 
 void FBatchPrinter::PrintObjects(TConstArrayView<FStructView> Objects)
@@ -947,7 +955,7 @@ void FBatchPrinter::PrintStructSchema(const FStructSchema& Struct, FSchemaBatchI
 		TextBuilder.AddLeaf(GLiterals.DeclaredSuper, *FPrintId<128>(Ids, SuperSchema.Type));
 	}
 
-	TextBuilder.BeginRange(GLiterals.Members);
+	TextBuilder.BeginStruct(GLiterals.Members);
 	TUtf8StringBuilder<256> Buf;
 	while (Reader.HasMore())
 	{
@@ -959,7 +967,7 @@ void FBatchPrinter::PrintStructSchema(const FStructSchema& Struct, FSchemaBatchI
 		TextBuilder.AddValue(FUtf8StringView(Buf));
 		Buf.Reset();
 	}
-	TextBuilder.EndRange();
+	TextBuilder.EndStruct();
 
 	TextBuilder.EndStruct();
 }
@@ -970,7 +978,7 @@ void FBatchPrinter::PrintEnumSchema(const FEnumSchema& Enum)
 	TextBuilder.AddLeaf(GLiterals.FlagMode, !!Enum.FlagMode);
 	TextBuilder.AddLeaf(GLiterals.Width, Enum.Width);
 
-	TextBuilder.BeginRange(GLiterals.Constants);
+	TextBuilder.BeginStruct(GLiterals.Constants);
 	TConstArrayView<FNameId> EnumNames = MakeConstArrayView(Enum.Footer, Enum.Num);
 	switch (Enum.Width)
 	{
@@ -987,7 +995,7 @@ void FBatchPrinter::PrintEnumSchema(const FEnumSchema& Enum)
 			PrintEnumConstants(EnumNames, GetConstants<uint64>(Enum), Enum.FlagMode);
 			break;
 	}
-	TextBuilder.EndRange();
+	TextBuilder.EndStruct();
 
 	TextBuilder.EndStruct();
 }
