nss: Struct scope for object-type-as-key.

Adding a struct key (the object's type) isn't valid as the parent is not suitable - "Objects" is a range. So an additional struct scope is required.

#rnxl
#Synced-CL 49920676

diff --git a/Engine/Plugins/Experimental/PlainProps/Source/Private/PlainPropsParse.cpp b/Engine/Plugins/Experimental/PlainProps/Source/Private/PlainPropsParse.cpp
--- a/Engine/Plugins/Experimental/PlainProps/Source/Private/PlainPropsParse.cpp
+++ b/Engine/Plugins/Experimental/PlainProps/Source/Private/PlainPropsParse.cpp
@@ -2100,12 +2100,25 @@ void FBatchParser::ParseObjects(TConstArrayView64<FToken> TokensView)
 	FScopedTokenReader ScopedReader(TokenIt);
 	while (ScopedReader.HasMore())
 	{
+		// Each item in the "Objects" range is a struct with one leaf, where the
+		// leaf's key is the object type and it's value a struct of the object;
+		//
+		//   [ {"<type>": { ... }}, {"<type>": ...etc
+		//
+		// So we've two structs and depths to fetch.
+
 		FToken Token;
 		if (!GrabToNextBeginStruct(Token, ScopedReader))
 		{
 			break;
 		}
 
+		FScopedTokenReader StructReader(TokenIt);
+		if (!GrabToNextBeginStruct(Token, StructReader))
+		{
+			break;
+		}
+
 		TOptional<FTypeTokens> TypeTokens = TokenizeType(Token.Value());
 		if (!TypeTokens)
 		{
diff --git a/Engine/Plugins/Experimental/PlainProps/Source/Private/PlainPropsPrint.cpp b/Engine/Plugins/Experimental/PlainProps/Source/Private/PlainPropsPrint.cpp
--- a/Engine/Plugins/Experimental/PlainProps/Source/Private/PlainPropsPrint.cpp
+++ b/Engine/Plugins/Experimental/PlainProps/Source/Private/PlainPropsPrint.cpp
@@ -1038,8 +1038,10 @@ void FBatchPrinter::PrintEnumConstants(
 void FMemberPrinter::PrintMembers(FStructView StructView)
 {
 	const FStructSchema& Schema = StructView.Schema.Resolve();
+	TextBuilder.BeginStruct();
 	TextBuilder.BeginStruct(*FPrintId<128>(Ids, Schema.Type));
 	PrintMembersInternal({ EMemberKind::Struct }, StructView, Schema);
+	TextBuilder.EndStruct();
 }
 
 void FMemberPrinter::PrintLeaf(FMemberId Id, FLeafView LeafView)
