nss: factored out structure text building from value handling

#rnx
#Synced-CL 49920676

diff --git a/Engine/Plugins/Experimental/PlainProps/Source/Private/PlainPropsPrint.cpp b/Engine/Plugins/Experimental/PlainProps/Source/Private/PlainPropsPrint.cpp
--- a/Engine/Plugins/Experimental/PlainProps/Source/Private/PlainPropsPrint.cpp
+++ b/Engine/Plugins/Experimental/PlainProps/Source/Private/PlainPropsPrint.cpp
@@ -107,11 +107,18 @@ void Print(FUtf8Builder& Out, T Value)
 }
 
 template<UnsignedIntegral T>
-void PrintHex(FUtf8Builder& Out, T Value)
+struct FPrintFormatHex
+{
+	explicit FPrintFormatHex(T InValue) : Value(InValue) {}
+	uint64 Value;
+};
+
+template<UnsignedIntegral T>
+void Print(FUtf8Builder& Out, FPrintFormatHex<T> Hex)
 {
 	constexpr size_t Size = sizeof(T) * 2;
 	char Buf[Size];
-	std::to_chars_result Result = std::to_chars(Buf, Buf + Size, Value, 16);
+	std::to_chars_result Result = std::to_chars(Buf, Buf + Size, Hex.Value, 16);
 	check(Result.ec == std::errc());
 	Out.Append(Buf, Result.ptr - Buf);
 }
@@ -328,6 +335,90 @@ static void PrintEnumValue(
 	while (Value != 0);
 }
 
+
+
+///////////////////////////////////////////////////////////////////////////////
+
+class IMarkupBuilder
+{
+public:
+	virtual void BeginDocument() = 0;
+	virtual void EndDocument() = 0;
+
+	virtual void PushStruct() = 0;
+	virtual void PopStruct() = 0;
+
+	virtual void PushRange() = 0;
+	virtual void PopRange() = 0;
+
+	virtual void WriteId(FUtf8StringView Id) = 0;
+	virtual void WriteValue(FUtf8StringView Value, bool bEscape=false) = 0;
+	virtual void WriteItem(FUtf8StringView Value) = 0;
+
+	virtual void AddComment(FUtf8StringView Comment) = 0;
+};
+
+
+
+///////////////////////////////////////////////////////////////////////////////
+
+class FMarkupBuilder
+	: public IMarkupBuilder
+{
+public:
+	void			BeginStruct()					{ PushStruct(); }
+	void			BeginStruct(FUtf8StringView Id)	{ WriteId(Id); BeginStruct(); }
+	void			EndStruct()						{ PopStruct(); }
+	void			BeginRange()					{ PushRange(); }
+	void			BeginRange(FUtf8StringView Id)	{ WriteId(Id); BeginRange(); }
+	void			EndRange()						{ PopRange(); }
+
+	// split calls for id and value pairs
+	void						AddId(FUtf8StringView Id) { WriteId(Id); }
+	template<typename T> void	AddValue(T Value);
+	void						AddValue(FUtf8StringView Value);
+	void						AddValue(const auto*) = delete;
+
+	// range items
+	template<typename T> void	AddItem(T Value);
+	void						AddItem(FUtf8StringView Value);
+	void						AddItem(const auto*) = delete;
+
+	// id value pairs
+	template<typename T>
+	void AddLeaf(FUtf8StringView Id, T Value)
+	{
+		WriteId(Id);
+		AddValue(Value);
+	}
+};
+
+template<typename T>
+void FMarkupBuilder::AddValue(T Value)
+{
+	TUtf8StringBuilder<256> Buffer;
+	Print(Buffer, Value);
+	WriteValue(Buffer);
+}
+
+void FMarkupBuilder::AddValue(FUtf8StringView Value)
+{
+	WriteValue(Value, true);
+}
+
+template<typename T>
+void FMarkupBuilder::AddItem(T Value)
+{
+	TUtf8StringBuilder<256> Buffer;
+	Print(Buffer, Value);
+	WriteItem(Buffer);
+}
+
+void FMarkupBuilder::AddItem(FUtf8StringView Value)
+{
+	WriteItem(Value);
+}
+
 ///////////////////////////////////////////////////////////////////////////////
 
 struct FMemberSchemaView
@@ -406,52 +497,34 @@ private:
 
 ///////////////////////////////////////////////////////////////////////////////
 
-class FYamlBuilder
+class FYamlBuilder final
+	: public FMarkupBuilder
 {
 public:
 	FYamlBuilder(FUtf8Builder& InStringBuilder);
 	~FYamlBuilder();
 
-	void BeginDocument();
-	void EndDocument();
+	virtual void BeginDocument() override;
+	virtual void EndDocument() override;
 
-	void BeginStruct(FUtf8StringView Id);
-	void BeginStruct();
-	void EndStruct();
+	virtual void PushStruct() override;
+	virtual void PopStruct() override;
 
-	void BeginRange(FUtf8StringView Id);
-	void BeginRange();
-	void EndRange();
+	virtual void PushRange() override;
+	virtual void PopRange() override;
 
 	// split calls for id and value pairs
-	void AddId(FUtf8StringView Id);
-	template<typename T>
-	void AddValue(T Value);
-	template<UnsignedIntegral T>
-	void AddValueHex(T Value);
-
-	// id value pairs
-	template<typename T>
-	void AddLeaf(FUtf8StringView Id, T Value) { AddId(Id); AddValue(Value); }
+	virtual void WriteId(FUtf8StringView Id) override;
+	virtual void WriteValue(FUtf8StringView Value, bool bEscaped=false) override;
 
 	// range items
-	template<typename T>
-	void AddItem(T Value);
-	template<UnsignedIntegral T>
-	void AddItemHex(T Value);
+	virtual void WriteItem(FUtf8StringView Value) override;
 
-	void AddComment(FUtf8StringView Comment);
+	virtual void AddComment(FUtf8StringView Comment) override;
 
 private:
 	void AppendNewLine();
-	void AppendIndentation();
-	void AppendIdentifier(FUtf8StringView Id);
-	template<typename T>
-	void AppendValue(T Value);
-	template<>
-	void AppendValue(FUtf8StringView Value);
-	template<UnsignedIntegral T>
-	void AppendValueHex(T Value);
+	void PopStack();
 
 	struct FStackInfo
 	{
@@ -461,8 +534,7 @@ private:
 
 	FUtf8Builder& Text;
 	TArray<FStackInfo, TInlineAllocator<32>> Stack;
-	uint32 IndentLevel = 0;
-	bool IsNewLine = true;
+	uint32 LineStart = 0;
 };
 
 ///////////////////////////////////////////////////////////////////////////////
@@ -657,6 +729,8 @@ FYamlBuilderPtr MakeYamlBuilder(FUtf8Builder& StringBuilder)
 
 ///////////////////////////////////////////////////////////////////////////////
 
+static void PrintQuotedString(FUtf8Builder& Out, FUtf8StringView Value);
+
 FYamlBuilder::FYamlBuilder(FUtf8Builder& InStringBuilder)
 : Text(InStringBuilder)
 {
@@ -671,160 +745,123 @@ FYamlBuilder::~FYamlBuilder()
 
 void FYamlBuilder::BeginDocument()
 {
-	Text << "---";
-	IsNewLine = false;
-	AppendNewLine();
+	Text << "---\n";
 	Stack.Emplace();
 }
 
 void FYamlBuilder::EndDocument()
 {
-	AppendNewLine();
 	Text << "...";
 	Stack.Pop(EAllowShrinking::No);
 }
 
-void FYamlBuilder::BeginStruct(FUtf8StringView Id)
+void FYamlBuilder::PushStruct()
 {
-	AppendNewLine();
-	AppendIndentation();
-	AppendIdentifier(Id);
 	Stack.Last().IsEmpty = false;
-	Stack.Emplace();
-	++IndentLevel;
-}
+	bool bNeedNewLine = Stack.Last().IsInStruct;
 
-void FYamlBuilder::BeginStruct()
-{
-	AppendNewLine();
-	AppendIndentation();
-	Stack.Last().IsEmpty = false;
 	Stack.Emplace();
-	++IndentLevel;
-}
 
-void FYamlBuilder::EndStruct()
-{
-	--IndentLevel;
-	if (Stack.Last().IsEmpty)
+	// Current parser needs struct-type range items to be on a new line.
+	bNeedNewLine = true;
+
+	LineStart = Text.Len();
+	if (bNeedNewLine)
 	{
-		Text << " {}";
-		IsNewLine = false;
+		AppendNewLine();
 	}
-	Stack.Pop(EAllowShrinking::No);
 }
 
-void FYamlBuilder::BeginRange(FUtf8StringView Id)
+void FYamlBuilder::PopStruct()
 {
-	AppendNewLine();
-	AppendIndentation();
-	AppendIdentifier(Id);
-	Stack.Last().IsEmpty = false;
-	Stack.Emplace_GetRef().IsInStruct = false;
-	++IndentLevel;
+	PopStack();
 }
 
-void FYamlBuilder::BeginRange()
+void FYamlBuilder::PushRange()
 {
-	AppendNewLine();
-	AppendIndentation();
 	Stack.Last().IsEmpty = false;
-	Stack.Emplace_GetRef().IsInStruct = false;
-	++IndentLevel;
-}
+	Stack.Emplace(FStackInfo{
+		.IsInStruct = false,
+	});
 
-void FYamlBuilder::EndRange()
-{
-	if (Stack.Last().IsEmpty)
-	{
-		Text << " []";
-		IsNewLine = false;
-	}
-	Stack.Pop(EAllowShrinking::No);
-	--IndentLevel;
+	AppendNewLine();
 }
 
-void FYamlBuilder::AddId(FUtf8StringView Id)
+void FYamlBuilder::PopRange()
 {
-	check(!Stack.IsEmpty() && Stack.Last().IsInStruct);
-
-	AppendNewLine();
-	AppendIndentation();
-	AppendIdentifier(Id);
-	IsNewLine = false;
-	Stack.Last().IsEmpty = false;
+	PopStack();
 }
 
-template<typename T>
-void FYamlBuilder::AddValue(T Value)
+void FYamlBuilder::WriteId(FUtf8StringView Id)
 {
-	check(!Stack.IsEmpty() && Stack.Last().IsInStruct);
+	PrintQuotedString(Text, Id);
+	Text.AppendChar(':');
 
-	Text.AppendChar(' ');
-	AppendValue(Value);
-	IsNewLine = false;
 	Stack.Last().IsEmpty = false;
 }
 
-template<UnsignedIntegral T>
-void FYamlBuilder::AddValueHex(T Value)
+void FYamlBuilder::WriteValue(FUtf8StringView Value, bool bEscaped)
 {
-	check(!Stack.IsEmpty() && Stack.Last().IsInStruct);
+	if (bEscaped)
+	{
+		Text << ' ';
+		PrintQuotedString(Text, Value);
+	}
+	else
+	{
+		Text << " \'";
+		Text << Value;
+		Text.AppendChar('\'');
+	}
 
-	Text.AppendChar(' ');
-	AppendValueHex(Value);
-	IsNewLine = false;
 	Stack.Last().IsEmpty = false;
-}
-
-template<typename T>
-void FYamlBuilder::AddItem(T Value)
-{
-	check(!Stack.IsEmpty() && Stack.Last().IsInStruct == false);
 
 	AppendNewLine();
-	AppendIndentation();
-	AppendValue(Value);
-	IsNewLine = false;
-	Stack.Last().IsEmpty = false;
 }
 
-template<UnsignedIntegral T>
-void FYamlBuilder::AddItemHex(T Value)
+void FYamlBuilder::WriteItem(FUtf8StringView Value)
 {
-	AppendNewLine();
-	AppendIndentation();
-	AppendValueHex(Value);
-	IsNewLine = false;
-	Stack.Last().IsEmpty = false;
+	check(Stack.Last().IsInStruct == false);
+	WriteValue(Value, true);
 }
 
 void FYamlBuilder::AddComment(FUtf8StringView Comment)
 {
+	Text.RemoveSuffix(Text.Len() - LineStart);
 	Text << " #" << Comment;
+
 	AppendNewLine();
 }
 
 void FYamlBuilder::AppendNewLine()
 {
-	if (!IsNewLine)
+	LineStart = Text.Len();
+
+	Text << '\n';
+	for (int32 Index = 0, Num = Stack.Num() - 2; Index < Num; ++Index)
 	{
-		Text << '\n';
-		IsNewLine = true;
+		Text << "  ";
 	}
-}
 
-void FYamlBuilder::AppendIndentation()
-{
-	for (uint32 I = 0; I < 2*IndentLevel; ++I)
+	if (Stack.Last().IsInStruct == false)
 	{
-		Text.AppendChar(' ');
+		Text << "- ";
 	}
-	if (!Stack.Last().IsInStruct)
+}
+
+void FYamlBuilder::PopStack()
+{
+	FStackInfo Last = Stack.Last();
+	Stack.Pop(EAllowShrinking::No);
+
+	Text.RemoveSuffix(Text.Len() - LineStart);
+
+	if (Last.IsEmpty)
 	{
-		Text << "- ";
+		Text << (Last.IsInStruct ? " {}" : " []");
 	}
-	IsNewLine = false;
+
+	AppendNewLine();
 }
 
 static void PrintQuotedString(FUtf8Builder& Out, FUtf8StringView Value)
@@ -852,34 +889,6 @@ static void PrintQuotedString(FUtf8Builder& Out, FUtf8StringView Value)
 	Out << '\"';
 }
 
-void FYamlBuilder::AppendIdentifier(FUtf8StringView Id)
-{
-	PrintQuotedString(Text, Id);
-	Text.AppendChar(':');
-	IsNewLine = false;
-}
-
-template<typename T>
-void FYamlBuilder::AppendValue(T Value)
-{
-	Text.AppendChar('\'');
-	Print(Text, Value);
-	Text.AppendChar('\'');
-}
-
-template<>
-void FYamlBuilder::AppendValue(FUtf8StringView Value)
-{
-	PrintQuotedString(Text, Value);
-}
-
-template<UnsignedIntegral T>
-void FYamlBuilder::AppendValueHex(T Value)
-{
-	Text.AppendChar('\'');
-	PrintHex(Text, Value);
-	Text.AppendChar('\'');
-}
 
 ///////////////////////////////////////////////////////////////////////////////
 
@@ -1085,10 +1094,10 @@ void FMemberPrinter::PrintLeaf(FMemberId Id, FLeafView LeafView)
 	case ELeafType::Hex:
 		switch (LeafView.Leaf.Width)
 		{
-			case ELeafWidth::B8:	TextBuilder.AddValueHex(LeafView.AsHex8());  break;
-			case ELeafWidth::B16:	TextBuilder.AddValueHex(LeafView.AsHex16()); break;
-			case ELeafWidth::B32:	TextBuilder.AddValueHex(LeafView.AsHex32()); break;
-			case ELeafWidth::B64:	TextBuilder.AddValueHex(LeafView.AsHex64()); break;
+			case ELeafWidth::B8:	TextBuilder.AddValue(FPrintFormatHex<uint8 >{LeafView.AsHex8()});  break;
+			case ELeafWidth::B16:	TextBuilder.AddValue(FPrintFormatHex<uint16>{LeafView.AsHex16()}); break;
+			case ELeafWidth::B32:	TextBuilder.AddValue(FPrintFormatHex<uint32>{LeafView.AsHex32()}); break;
+			case ELeafWidth::B64:	TextBuilder.AddValue(FPrintFormatHex<uint64>{LeafView.AsHex64()}); break;
 		}
 		break;
 	case ELeafType::Enum:
@@ -1187,10 +1196,10 @@ void FMemberPrinter::PrintLeaves(FUnpackedLeafType Leaf, const FLeafRangeView& L
 	case ELeafType::Hex:
 		switch (Leaf.Width)
 		{
-			case ELeafWidth::B8:  for (const uint8&  U : LeafRange.AsHex8s())  { TextBuilder.AddItemHex(U); } break;
-			case ELeafWidth::B16: for (const uint16& U : LeafRange.AsHex16s()) { TextBuilder.AddItemHex(U); } break;
-			case ELeafWidth::B32: for (const uint32& U : LeafRange.AsHex32s()) { TextBuilder.AddItemHex(U); } break;
-			case ELeafWidth::B64: for (const uint64& U : LeafRange.AsHex64s()) { TextBuilder.AddItemHex(U); } break;
+			case ELeafWidth::B8:  for (uint8  U : LeafRange.AsHex8s())  { TextBuilder.AddItem(FPrintFormatHex<uint8 >{U}); } break;
+			case ELeafWidth::B16: for (uint16 U : LeafRange.AsHex16s()) { TextBuilder.AddItem(FPrintFormatHex<uint16>{U}); } break;
+			case ELeafWidth::B32: for (uint32 U : LeafRange.AsHex32s()) { TextBuilder.AddItem(FPrintFormatHex<uint32>{U}); } break;
+			case ELeafWidth::B64: for (uint64 U : LeafRange.AsHex64s()) { TextBuilder.AddItem(FPrintFormatHex<uint64>{U}); } break;
 		}
 		break;
 	case ELeafType::Enum:
