nss: reduce AddValue() permutation and gather a loose print function with the others.

#rnx

#Synced-CL 49920676

diff --git a/Engine/Plugins/Experimental/PlainProps/Source/Private/PlainPropsPrint.cpp b/Engine/Plugins/Experimental/PlainProps/Source/Private/PlainPropsPrint.cpp
--- a/Engine/Plugins/Experimental/PlainProps/Source/Private/PlainPropsPrint.cpp
+++ b/Engine/Plugins/Experimental/PlainProps/Source/Private/PlainPropsPrint.cpp
@@ -107,11 +107,18 @@ void Print(FUtf8Builder& Out, T Value)
 }
 
 template<UnsignedIntegral T>
-void PrintHex(FUtf8Builder& Out, T Value)
+struct FPrintFormatHex
+{
+	explicit FPrintFormatHex(T InValue) : Value(InValue) {}
+	uint64 Value;
+};
+
+template<UnsignedIntegral T>
+void Print(FUtf8Builder& Out, FPrintFormatHex<T> Hex)
 {
 	constexpr size_t Size = sizeof(T) * 2;
 	char Buf[Size];
-	std::to_chars_result Result = std::to_chars(Buf, Buf + Size, Value, 16);
+	std::to_chars_result Result = std::to_chars(Buf, Buf + Size, Hex.Value, 16);
 	check(Result.ec == std::errc());
 	Out.Append(Buf, Result.ptr - Buf);
 }
@@ -172,6 +179,32 @@ void Print(FUtf8Builder& Out, char32_t Value)
 	}
 }
 
+static void PrintQuotedString(FUtf8Builder& Out, FUtf8StringView Value)
+{
+	FUtf8StringView Verbatim = FAsciiSet::FindPrefixWithout(Value, EscapeSet | "'");
+	if (Verbatim.Len() == Value.Len())
+	{
+		Out << '\'' << Value << '\'';
+		return;
+	}
+
+	Out << '\"';
+	while (!Value.IsEmpty())
+	{
+		Out << Verbatim;
+		Value.RightChopInline(Verbatim.Len());
+		FUtf8StringView Escape = FAsciiSet::FindPrefixWith(Value, EscapeSet);
+		for (UTF8CHAR Char : Escape)
+		{
+			EscapeChar(Out, Char);
+		}
+		Value.RightChopInline(Escape.Len());
+		Verbatim = FAsciiSet::FindPrefixWithout(Value, EscapeSet);
+	}
+	Out << '\"';
+}
+
+
 ///////////////////////////////////////////////////////////////////////////////
 
 static void PrintRangeType(FUtf8Builder& Out, FRangeType Type)
@@ -427,8 +460,6 @@ public:
 	void AddId(FUtf8StringView Id);
 	template<typename T>
 	void AddValue(T Value);
-	template<UnsignedIntegral T>
-	void AddValueHex(T Value);
 
 	// id value pairs
 	template<typename T>
@@ -437,8 +468,6 @@ public:
 	// range items
 	template<typename T>
 	void AddItem(T Value);
-	template<UnsignedIntegral T>
-	void AddItemHex(T Value);
 
 	void AddComment(FUtf8StringView Comment);
 
@@ -766,17 +795,6 @@ void FYamlBuilder::AddValue(T Value)
 	Stack.Last().IsEmpty = false;
 }
 
-template<UnsignedIntegral T>
-void FYamlBuilder::AddValueHex(T Value)
-{
-	check(!Stack.IsEmpty() && Stack.Last().IsInStruct);
-
-	Text.AppendChar(' ');
-	AppendValueHex(Value);
-	IsNewLine = false;
-	Stack.Last().IsEmpty = false;
-}
-
 template<typename T>
 void FYamlBuilder::AddItem(T Value)
 {
@@ -789,16 +807,6 @@ void FYamlBuilder::AddItem(T Value)
 	Stack.Last().IsEmpty = false;
 }
 
-template<UnsignedIntegral T>
-void FYamlBuilder::AddItemHex(T Value)
-{
-	AppendNewLine();
-	AppendIndentation();
-	AppendValueHex(Value);
-	IsNewLine = false;
-	Stack.Last().IsEmpty = false;
-}
-
 void FYamlBuilder::AddComment(FUtf8StringView Comment)
 {
 	Text << " #" << Comment;
@@ -827,31 +835,6 @@ void FYamlBuilder::AppendIndentation()
 	IsNewLine = false;
 }
 
-static void PrintQuotedString(FUtf8Builder& Out, FUtf8StringView Value)
-{
-	FUtf8StringView Verbatim = FAsciiSet::FindPrefixWithout(Value, EscapeSet | "'");
-	if (Verbatim.Len() == Value.Len())
-	{
-		Out << '\'' << Value << '\'';
-		return;
-	}
-
-	Out << '\"';
-	while (!Value.IsEmpty())
-	{
-		Out << Verbatim;
-		Value.RightChopInline(Verbatim.Len());
-		FUtf8StringView Escape = FAsciiSet::FindPrefixWith(Value, EscapeSet);
-		for (UTF8CHAR Char : Escape)
-		{
-			EscapeChar(Out, Char);
-		}
-		Value.RightChopInline(Escape.Len());
-		Verbatim = FAsciiSet::FindPrefixWithout(Value, EscapeSet);
-	}
-	Out << '\"';
-}
-
 void FYamlBuilder::AppendIdentifier(FUtf8StringView Id)
 {
 	PrintQuotedString(Text, Id);
@@ -1085,10 +1068,10 @@ void FMemberPrinter::PrintLeaf(FMemberId Id, FLeafView LeafView)
 	case ELeafType::Hex:
 		switch (LeafView.Leaf.Width)
 		{
-			case ELeafWidth::B8:	TextBuilder.AddValueHex(LeafView.AsHex8());  break;
-			case ELeafWidth::B16:	TextBuilder.AddValueHex(LeafView.AsHex16()); break;
-			case ELeafWidth::B32:	TextBuilder.AddValueHex(LeafView.AsHex32()); break;
-			case ELeafWidth::B64:	TextBuilder.AddValueHex(LeafView.AsHex64()); break;
+			case ELeafWidth::B8:	TextBuilder.AddValue(FPrintFormatHex<uint8 >{LeafView.AsHex8()});  break;
+			case ELeafWidth::B16:	TextBuilder.AddValue(FPrintFormatHex<uint16>{LeafView.AsHex16()}); break;
+			case ELeafWidth::B32:	TextBuilder.AddValue(FPrintFormatHex<uint32>{LeafView.AsHex32()}); break;
+			case ELeafWidth::B64:	TextBuilder.AddValue(FPrintFormatHex<uint64>{LeafView.AsHex64()}); break;
 		}
 		break;
 	case ELeafType::Enum:
@@ -1187,10 +1170,10 @@ void FMemberPrinter::PrintLeaves(FUnpackedLeafType Leaf, const FLeafRangeView& L
 	case ELeafType::Hex:
 		switch (Leaf.Width)
 		{
-			case ELeafWidth::B8:  for (const uint8&  U : LeafRange.AsHex8s())  { TextBuilder.AddItemHex(U); } break;
-			case ELeafWidth::B16: for (const uint16& U : LeafRange.AsHex16s()) { TextBuilder.AddItemHex(U); } break;
-			case ELeafWidth::B32: for (const uint32& U : LeafRange.AsHex32s()) { TextBuilder.AddItemHex(U); } break;
-			case ELeafWidth::B64: for (const uint64& U : LeafRange.AsHex64s()) { TextBuilder.AddItemHex(U); } break;
+			case ELeafWidth::B8:  for (uint8  U : LeafRange.AsHex8s())  { TextBuilder.AddItem(FPrintFormatHex<uint8> {U}); } break;
+			case ELeafWidth::B16: for (uint16 U : LeafRange.AsHex16s()) { TextBuilder.AddItem(FPrintFormatHex<uint16>{U}); } break;
+			case ELeafWidth::B32: for (uint32 U : LeafRange.AsHex32s()) { TextBuilder.AddItem(FPrintFormatHex<uint32>{U}); } break;
+			case ELeafWidth::B64: for (uint64 U : LeafRange.AsHex64s()) { TextBuilder.AddItem(FPrintFormatHex<uint64>{U}); } break;
 		}
 		break;
 	case ELeafType::Enum:
