nss: factored out structure text building from value handling

#rnx

#Synced-CL 49920676

diff --git a/Engine/Plugins/Experimental/PlainProps/Source/Private/PlainPropsPrint.cpp b/Engine/Plugins/Experimental/PlainProps/Source/Private/PlainPropsPrint.cpp
--- a/Engine/Plugins/Experimental/PlainProps/Source/Private/PlainPropsPrint.cpp
+++ b/Engine/Plugins/Experimental/PlainProps/Source/Private/PlainPropsPrint.cpp
@@ -361,6 +361,92 @@ static void PrintEnumValue(
 	while (Value != 0);
 }
 
+
+
+///////////////////////////////////////////////////////////////////////////////
+
+class IMarkupBuilder
+{
+public:
+	virtual void PushDocument() = 0;
+	virtual void PopDocument() = 0;
+
+	virtual void PushStruct() = 0;
+	virtual void PopStruct() = 0;
+	virtual void WriteId(FUtf8StringView Id) = 0;
+	virtual void WriteValue(FUtf8StringView Value, bool bEscape=false) = 0;
+
+	virtual void PushRange() = 0;
+	virtual void PopRange() = 0;
+	virtual void WriteItem(FUtf8StringView Value) = 0;
+
+	virtual void WriteComment(FUtf8StringView Comment) = 0;
+};
+
+
+
+///////////////////////////////////////////////////////////////////////////////
+
+class FMarkupBuilder
+	: public IMarkupBuilder
+{
+public:
+	void		BeginDocument()					{ PushDocument(); }
+	void		EndDocument()					{ PopDocument(); }
+	void		BeginStruct()					{ PushStruct(); }
+	void		BeginStruct(FUtf8StringView Id)	{ WriteId(Id); BeginStruct(); }
+	void		EndStruct()						{ PopStruct(); }
+	void		BeginRange()					{ PushRange(); }
+	void		BeginRange(FUtf8StringView Id)	{ WriteId(Id); BeginRange(); }
+	void		EndRange()						{ PopRange(); }
+	void		AddComment(FUtf8StringView Str) { WriteComment(Str); }
+
+	// split calls for id and value pairs
+	void						AddId(FUtf8StringView Id) { WriteId(Id); }
+	template<typename T> void	AddValue(T Value);
+	void						AddValue(FUtf8StringView Value);
+	void						AddValue(const auto*) = delete;
+
+	// range items
+	template<typename T> void	AddItem(T Value);
+	void						AddItem(FUtf8StringView Value);
+	void						AddItem(const auto*) = delete;
+
+	// id value pairs
+	template<typename T>
+	void AddLeaf(FUtf8StringView Id, T Value)
+	{
+		WriteId(Id);
+		AddValue(Value);
+	}
+};
+
+template<typename T>
+void FMarkupBuilder::AddValue(T Value)
+{
+	TUtf8StringBuilder<256> Buffer;
+	Print(Buffer, Value);
+	WriteValue(Buffer);
+}
+
+void FMarkupBuilder::AddValue(FUtf8StringView Value)
+{
+	WriteValue(Value, true);
+}
+
+template<typename T>
+void FMarkupBuilder::AddItem(T Value)
+{
+	TUtf8StringBuilder<256> Buffer;
+	Print(Buffer, Value);
+	WriteItem(Buffer);
+}
+
+void FMarkupBuilder::AddItem(FUtf8StringView Value)
+{
+	WriteItem(Value);
+}
+
 ///////////////////////////////////////////////////////////////////////////////
 
 struct FMemberSchemaView
@@ -439,48 +525,26 @@ private:
 
 ///////////////////////////////////////////////////////////////////////////////
 
-class FYamlBuilder
+class FYamlBuilder final
+	: public FMarkupBuilder
 {
 public:
 	FYamlBuilder(FUtf8Builder& InStringBuilder);
 	~FYamlBuilder();
 
-	void BeginDocument();
-	void EndDocument();
-
-	void BeginStruct(FUtf8StringView Id);
-	void BeginStruct();
-	void EndStruct();
-
-	void BeginRange(FUtf8StringView Id);
-	void BeginRange();
-	void EndRange();
-
-	// split calls for id and value pairs
-	void AddId(FUtf8StringView Id);
-	template<typename T>
-	void AddValue(T Value);
-
-	// id value pairs
-	template<typename T>
-	void AddLeaf(FUtf8StringView Id, T Value) { AddId(Id); AddValue(Value); }
-
-	// range items
-	template<typename T>
-	void AddItem(T Value);
-
-	void AddComment(FUtf8StringView Comment);
-
 private:
+	virtual void PushDocument() override;
+	virtual void PopDocument() override;
+	virtual void PushStruct() override;
+	virtual void PopStruct() override;
+	virtual void PushRange() override;
+	virtual void PopRange() override;
+	virtual void WriteId(FUtf8StringView Id) override;
+	virtual void WriteValue(FUtf8StringView Value, bool bEscaped=false) override;
+	virtual void WriteItem(FUtf8StringView Value) override;
+	virtual void WriteComment(FUtf8StringView Comment) override;
 	void AppendNewLine();
-	void AppendIndentation();
-	void AppendIdentifier(FUtf8StringView Id);
-	template<typename T>
-	void AppendValue(T Value);
-	template<>
-	void AppendValue(FUtf8StringView Value);
-	template<UnsignedIntegral T>
-	void AppendValueHex(T Value);
+	void PopStack();
 
 	struct FStackInfo
 	{
@@ -490,8 +554,7 @@ private:
 
 	FUtf8Builder& Text;
 	TArray<FStackInfo, TInlineAllocator<32>> Stack;
-	uint32 IndentLevel = 0;
-	bool IsNewLine = true;
+	uint32 LineStart = 0;
 };
 
 ///////////////////////////////////////////////////////////////////////////////
@@ -698,172 +761,136 @@ FYamlBuilder::~FYamlBuilder()
 	check(Stack.IsEmpty());
 }
 
-void FYamlBuilder::BeginDocument()
+void FYamlBuilder::PushDocument()
 {
-	Text << "---";
-	IsNewLine = false;
-	AppendNewLine();
+	Text << "---\n";
 	Stack.Emplace();
 }
 
-void FYamlBuilder::EndDocument()
+void FYamlBuilder::PopDocument()
 {
-	AppendNewLine();
 	Text << "...";
 	Stack.Pop(EAllowShrinking::No);
 }
 
-void FYamlBuilder::BeginStruct(FUtf8StringView Id)
+void FYamlBuilder::PushStruct()
 {
-	AppendNewLine();
-	AppendIndentation();
-	AppendIdentifier(Id);
 	Stack.Last().IsEmpty = false;
-	Stack.Emplace();
-	++IndentLevel;
-}
+	bool bNeedNewLine = Stack.Last().IsInStruct;
 
-void FYamlBuilder::BeginStruct()
-{
-	AppendNewLine();
-	AppendIndentation();
-	Stack.Last().IsEmpty = false;
 	Stack.Emplace();
-	++IndentLevel;
-}
 
-void FYamlBuilder::EndStruct()
-{
-	--IndentLevel;
-	if (Stack.Last().IsEmpty)
+	// Current parser needs struct-type range items to be on a new line.
+	bNeedNewLine = true;
+
+	LineStart = Text.Len();
+	if (bNeedNewLine)
 	{
-		Text << " {}";
-		IsNewLine = false;
+		AppendNewLine();
 	}
-	Stack.Pop(EAllowShrinking::No);
 }
 
-void FYamlBuilder::BeginRange(FUtf8StringView Id)
+void FYamlBuilder::PopStruct()
 {
-	AppendNewLine();
-	AppendIndentation();
-	AppendIdentifier(Id);
-	Stack.Last().IsEmpty = false;
-	Stack.Emplace_GetRef().IsInStruct = false;
-	++IndentLevel;
+	PopStack();
 }
 
-void FYamlBuilder::BeginRange()
+void FYamlBuilder::PushRange()
 {
-	AppendNewLine();
-	AppendIndentation();
 	Stack.Last().IsEmpty = false;
-	Stack.Emplace_GetRef().IsInStruct = false;
-	++IndentLevel;
+	Stack.Emplace(FStackInfo{
+		.IsInStruct = false,
+	});
+
+	AppendNewLine();
 }
 
-void FYamlBuilder::EndRange()
+void FYamlBuilder::PopRange()
 {
-	if (Stack.Last().IsEmpty)
-	{
-		Text << " []";
-		IsNewLine = false;
-	}
-	Stack.Pop(EAllowShrinking::No);
-	--IndentLevel;
+	PopStack();
 }
 
-void FYamlBuilder::AddId(FUtf8StringView Id)
+void FYamlBuilder::WriteId(FUtf8StringView Id)
 {
-	check(!Stack.IsEmpty() && Stack.Last().IsInStruct);
+	PrintQuotedString(Text, Id);
+	Text.AppendChar(':');
 
-	AppendNewLine();
-	AppendIndentation();
-	AppendIdentifier(Id);
-	IsNewLine = false;
 	Stack.Last().IsEmpty = false;
 }
 
-template<typename T>
-void FYamlBuilder::AddValue(T Value)
+void FYamlBuilder::WriteValue(FUtf8StringView Value, bool bEscaped)
 {
-	check(!Stack.IsEmpty() && Stack.Last().IsInStruct);
+	if (bEscaped)
+	{
+		Text << ' ';
+		PrintQuotedString(Text, Value);
+	}
+	else
+	{
+		Text << " \'";
+		Text << Value;
+		Text.AppendChar('\'');
+	}
 
-	Text.AppendChar(' ');
-	AppendValue(Value);
-	IsNewLine = false;
 	Stack.Last().IsEmpty = false;
+
+	AppendNewLine();
 }
 
-template<typename T>
-void FYamlBuilder::AddItem(T Value)
+void FYamlBuilder::WriteItem(FUtf8StringView Value)
 {
-	check(!Stack.IsEmpty() && Stack.Last().IsInStruct == false);
-
-	AppendNewLine();
-	AppendIndentation();
-	AppendValue(Value);
-	IsNewLine = false;
-	Stack.Last().IsEmpty = false;
+	check(Stack.Last().IsInStruct == false);
+	WriteValue(Value, true);
 }
 
-void FYamlBuilder::AddComment(FUtf8StringView Comment)
+void FYamlBuilder::WriteComment(FUtf8StringView Comment)
 {
+	Text.RemoveSuffix(Text.Len() - LineStart);
 	Text << " #" << Comment;
+
 	AppendNewLine();
 }
 
 void FYamlBuilder::AppendNewLine()
 {
-	if (!IsNewLine)
-	{
-		Text << '\n';
-		IsNewLine = true;
-	}
-}
+	LineStart = Text.Len();
 
-void FYamlBuilder::AppendIndentation()
-{
-	for (uint32 I = 0; I < 2*IndentLevel; ++I)
+	Text << '\n';
+
+	int32 IndentNum = (Stack.Num() - 2) << 1;
+	Text.Reserve(LineStart + IndentNum);
+	for (const FUtf8StringView Spaces("                                    ");;)
 	{
-		Text.AppendChar(' ');
+		FUtf8StringView Trimmed = Spaces.Left(IndentNum);
+		Text << Trimmed;
+		if (IndentNum -= Trimmed.Len(); IndentNum <= 0)
+		{
+			break;
+		}
 	}
-	if (!Stack.Last().IsInStruct)
+
+	if (Stack.Last().IsInStruct == false)
 	{
 		Text << "- ";
 	}
-	IsNewLine = false;
 }
 
-void FYamlBuilder::AppendIdentifier(FUtf8StringView Id)
+void FYamlBuilder::PopStack()
 {
-	PrintQuotedString(Text, Id);
-	Text.AppendChar(':');
-	IsNewLine = false;
-}
+	FStackInfo Last = Stack.Last();
+	Stack.Pop(EAllowShrinking::No);
 
-template<typename T>
-void FYamlBuilder::AppendValue(T Value)
-{
-	Text.AppendChar('\'');
-	Print(Text, Value);
-	Text.AppendChar('\'');
-}
+	Text.RemoveSuffix(Text.Len() - LineStart);
 
-template<>
-void FYamlBuilder::AppendValue(FUtf8StringView Value)
-{
-	PrintQuotedString(Text, Value);
-}
+	if (Last.IsEmpty)
+	{
+		Text << (Last.IsInStruct ? " {}" : " []");
+	}
 
-template<UnsignedIntegral T>
-void FYamlBuilder::AppendValueHex(T Value)
-{
-	Text.AppendChar('\'');
-	PrintHex(Text, Value);
-	Text.AppendChar('\'');
+	AppendNewLine();
 }
 
+
 ///////////////////////////////////////////////////////////////////////////////
 
 FBatchPrinter::FBatchPrinter(FYamlBuilder& InTextBuilder, const FBatchIds& InIds)
