nss: A JSON printer

#rnx

#Synced-CL 49920676

diff --git a/Engine/Plugins/Experimental/PlainProps/Source/Private/PlainPropsPrint.cpp b/Engine/Plugins/Experimental/PlainProps/Source/Private/PlainPropsPrint.cpp
--- a/Engine/Plugins/Experimental/PlainProps/Source/Private/PlainPropsPrint.cpp
+++ b/Engine/Plugins/Experimental/PlainProps/Source/Private/PlainPropsPrint.cpp
@@ -179,12 +179,12 @@ void Print(FUtf8Builder& Out, char32_t Value)
 	}
 }
 
-static void PrintQuotedString(FUtf8Builder& Out, FUtf8StringView Value)
+static void PrintQuotedString(FUtf8Builder& Out, FUtf8StringView Value, UTF8CHAR VerbatimQuote='\'')
 {
 	FUtf8StringView Verbatim = FAsciiSet::FindPrefixWithout(Value, EscapeSet | "'");
 	if (Verbatim.Len() == Value.Len())
 	{
-		Out << '\'' << Value << '\'';
+		Out << VerbatimQuote << Value << VerbatimQuote;
 		return;
 	}
 
@@ -523,6 +523,194 @@ private:
 	const FMemberId*		GetMemberNames() const;
 };
 
+///////////////////////////////////////////////////////////////////////////////
+
+class FJsonBuilder final
+	: public FMarkupBuilder
+{
+public:
+					FJsonBuilder(FUtf8Builder& InBuilder);
+					~FJsonBuilder();
+private:
+	enum class EAction { Open, Add, Close };
+
+	struct FScope
+	{
+		uint8		bIsStruct : 1;
+		uint8		bHasItems : 1;
+		uint8		KeyValueTickTock : 1;
+		uint8		_Unused : 5;
+	};
+
+	using FScopes = TArray<FScope, TInlineAllocator<32>>;
+
+	virtual void	PushDocument() override;
+	virtual void	PopDocument() override;
+	virtual void	PushStruct() override;
+	virtual void	PopStruct() override;
+	virtual void	PushRange() override;
+	virtual void	PopRange() override;
+	virtual void	WriteId(FUtf8StringView Id) override;
+	virtual void	WriteValue(FUtf8StringView Value, bool bEscaped=false) override;
+	virtual void	WriteItem(FUtf8StringView Value) override;
+	virtual void	WriteComment(FUtf8StringView Comment) override;
+	void			Prologue(EAction Action);
+	FScopes			Scopes;
+	FUtf8Builder&	Builder;
+};
+
+FJsonBuilder::FJsonBuilder(FUtf8Builder& InBuilder)
+: Builder(InBuilder)
+{
+}
+
+FJsonBuilder::~FJsonBuilder()
+{
+	check(Scopes.IsEmpty());
+}
+
+void FJsonBuilder::Prologue(EAction Action)
+{
+	if (Scopes.IsEmpty())
+	{
+		return;
+	}
+
+	FScope& Scope = Scopes.Last();
+
+	bool bHadItems = Scope.bHasItems;
+	Scope.bHasItems = true;
+
+	bool bIndent = true;
+
+	switch (Action)
+	{
+	case EAction::Close:
+		Scopes.Pop(EAllowShrinking::No);
+		if (!bHadItems)
+		{
+			return;
+		}
+		Builder << '\n';
+		break;
+
+	case EAction::Open:
+		if (!Scope.bIsStruct && bHadItems)
+		{
+			Builder << ",";
+		}
+		bIndent = ((Scope.KeyValueTickTock & 1) == 0);
+		if (bIndent)
+		{
+			Builder << "\n";
+		}
+		else
+		Scope.KeyValueTickTock += (Scope.bIsStruct == true);
+		break;
+
+	case EAction::Add:
+		bIndent = ((Scope.KeyValueTickTock & 1) == 0);
+		if (bIndent)
+		{
+			Builder << (bHadItems ? ",\n" : "\n");
+		}
+		Scope.KeyValueTickTock += (Scope.bIsStruct == true);
+		break;
+	}
+
+	if (!bIndent)
+	{
+		return;
+	}
+
+	int32 IndentNum = Scopes.Num() << 1;
+	for (const FUtf8StringView Spaces("                                    ");;)
+	{
+		FUtf8StringView Trimmed = Spaces.Left(IndentNum);
+		Builder << Trimmed;
+		if (IndentNum -= Trimmed.Len(); IndentNum <= 0)
+		{
+			break;
+		}
+	}
+}
+
+void FJsonBuilder::PushDocument()
+{
+	PushStruct();
+}
+
+void FJsonBuilder::PopDocument()
+{
+	PopStruct();
+	Builder << '\n';
+}
+
+void FJsonBuilder::PushStruct()
+{
+	Prologue(EAction::Open);
+
+	Builder << '{';
+
+	Scopes.Emplace(FScope{
+		.bIsStruct = true,
+	});
+}
+
+void FJsonBuilder::PopStruct()
+{
+	Prologue(EAction::Close);
+	Builder << '}';
+}
+
+void FJsonBuilder::PushRange()
+{
+	Prologue(EAction::Open);
+	Builder << '[';
+	Scopes.Emplace(FScope{});
+}
+
+void FJsonBuilder::PopRange()
+{
+	Prologue(EAction::Close);
+	Builder << ']';
+}
+
+void FJsonBuilder::WriteId(FUtf8StringView Id)
+{
+	Prologue(EAction::Add);
+
+	PrintQuotedString(Builder, Id);
+	Builder << " : ";
+}
+
+void FJsonBuilder::WriteValue(FUtf8StringView Value, bool bEscaped)
+{
+	Prologue(EAction::Add);
+
+	if (bEscaped)
+	{
+		PrintQuotedString(Builder, Value, '\"');
+		return;
+	}
+
+	Builder << '\"';
+	Builder << Value;
+	Builder << '\"';
+}
+
+void FJsonBuilder::WriteItem(FUtf8StringView Value)
+{
+	Prologue(EAction::Add);
+	PrintQuotedString(Builder, Value, '\"');
+}
+
+void FJsonBuilder::WriteComment(FUtf8StringView Comment)
+{
+}
+
+
+
 ///////////////////////////////////////////////////////////////////////////////
 
 class FYamlBuilder final
