nss: Stack -> Scope to match parser's YAML tokeniser.

#rnx

#Synced-CL 49920676

diff --git a/Engine/Plugins/Experimental/PlainProps/Source/Private/PlainPropsPrint.cpp b/Engine/Plugins/Experimental/PlainProps/Source/Private/PlainPropsPrint.cpp
--- a/Engine/Plugins/Experimental/PlainProps/Source/Private/PlainPropsPrint.cpp
+++ b/Engine/Plugins/Experimental/PlainProps/Source/Private/PlainPropsPrint.cpp
@@ -732,16 +732,17 @@ private:
 	virtual void WriteItem(FUtf8StringView Value) override;
 	virtual void WriteComment(FUtf8StringView Comment) override;
 	void AppendNewLine();
-	void PopStack();
+	void PopScope();
 
-	struct FStackInfo
+	struct FScopeInfo
 	{
-		bool IsEmpty = true;
-		bool IsInStruct = true;
+		uint8 HasItems : 1;
+		uint8 IsInStruct : 1;
+		uint8 _Unused : 6;
 	};
 
 	FUtf8Builder& Text;
-	TArray<FStackInfo, TInlineAllocator<32>> Stack;
+	TArray<FScopeInfo, TInlineAllocator<32>> Scope;
 	uint32 LineStart = 0;
 };
 
@@ -940,33 +941,37 @@ FYamlBuilderPtr MakeYamlBuilder(FUtf8Builder& StringBuilder)
 FYamlBuilder::FYamlBuilder(FUtf8Builder& InStringBuilder)
 : Text(InStringBuilder)
 {
-	Stack.Emplace();
+	Scope.Emplace(FScopeInfo{});
 }
 
 FYamlBuilder::~FYamlBuilder()
 {
-	Stack.Pop(EAllowShrinking::No);
-	check(Stack.IsEmpty());
+	Scope.Pop(EAllowShrinking::No);
+	check(Scope.IsEmpty());
 }
 
 void FYamlBuilder::PushDocument()
 {
 	Text << "---\n";
-	Stack.Emplace();
+	Scope.Emplace(FScopeInfo{
+		.IsInStruct = true,
+	});
 }
 
 void FYamlBuilder::PopDocument()
 {
 	Text << "...";
-	Stack.Pop(EAllowShrinking::No);
+	Scope.Pop(EAllowShrinking::No);
 }
 
 void FYamlBuilder::PushStruct()
 {
-	Stack.Last().IsEmpty = false;
-	bool bNeedNewLine = Stack.Last().IsInStruct;
+	Scope.Last().HasItems = true;
+	bool bNeedNewLine = Scope.Last().IsInStruct;
 
-	Stack.Emplace();
+	Scope.Emplace(FScopeInfo{
+		.IsInStruct = true,
+	});
 
 	// Current parser needs struct-type range items to be on a new line.
 	bNeedNewLine = true;
@@ -980,13 +985,13 @@ void FYamlBuilder::PushStruct()
 
 void FYamlBuilder::PopStruct()
 {
-	PopStack();
+	PopScope();
 }
 
 void FYamlBuilder::PushRange()
 {
-	Stack.Last().IsEmpty = false;
-	Stack.Emplace(FStackInfo{
+	Scope.Last().HasItems = true;
+	Scope.Emplace(FScopeInfo{
 		.IsInStruct = false,
 	});
 
@@ -995,7 +1000,7 @@ void FYamlBuilder::PushRange()
 
 void FYamlBuilder::PopRange()
 {
-	PopStack();
+	PopScope();
 }
 
 void FYamlBuilder::WriteId(FUtf8StringView Id)
@@ -1003,7 +1008,7 @@ void FYamlBuilder::WriteId(FUtf8StringView Id)
 	PrintQuotedString(Text, Id);
 	Text.AppendChar(':');
 
-	Stack.Last().IsEmpty = false;
+	Scope.Last().HasItems = true;
 }
 
 void FYamlBuilder::WriteValue(FUtf8StringView Value, bool bEscaped)
@@ -1020,14 +1025,14 @@ void FYamlBuilder::WriteValue(FUtf8StringView Value, bool bEscaped)
 		Text.AppendChar('\'');
 	}
 
-	Stack.Last().IsEmpty = false;
+	Scope.Last().HasItems = true;
 
 	AppendNewLine();
 }
 
 void FYamlBuilder::WriteItem(FUtf8StringView Value)
 {
-	check(Stack.Last().IsInStruct == false);
+	check(Scope.Last().IsInStruct == false);
 	WriteValue(Value, true);
 }
 
@@ -1045,7 +1050,7 @@ void FYamlBuilder::AppendNewLine()
 
 	Text << '\n';
 
-	int32 IndentNum = (Stack.Num() - 2) << 1;
+	int32 IndentNum = (Scope.Num() - 2) << 1;
 	Text.Reserve(LineStart + IndentNum);
 	for (const FUtf8StringView Spaces("                                    ");;)
 	{
@@ -1057,20 +1062,20 @@ void FYamlBuilder::AppendNewLine()
 		}
 	}
 
-	if (Stack.Last().IsInStruct == false)
+	if (Scope.Last().IsInStruct == false)
 	{
 		Text << "- ";
 	}
 }
 
-void FYamlBuilder::PopStack()
+void FYamlBuilder::PopScope()
 {
-	FStackInfo Last = Stack.Last();
-	Stack.Pop(EAllowShrinking::No);
+	FScopeInfo Last = Scope.Last();
+	Scope.Pop(EAllowShrinking::No);
 
 	Text.RemoveSuffix(Text.Len() - LineStart);
 
-	if (Last.IsEmpty)
+	if (!Last.HasItems)
 	{
 		Text << (Last.IsInStruct ? " {}" : " []");
 	}
