nss: Use the general markup interface for printing

#rnx

#Synced-CL 49920676

diff --git a/Engine/Plugins/Experimental/PlainProps/Source/Private/PlainPropsInternalPrint.h b/Engine/Plugins/Experimental/PlainProps/Source/Private/PlainPropsInternalPrint.h
--- a/Engine/Plugins/Experimental/PlainProps/Source/Private/PlainPropsInternalPrint.h
+++ b/Engine/Plugins/Experimental/PlainProps/Source/Private/PlainPropsInternalPrint.h
@@ -11,6 +11,7 @@
 namespace PlainProps
 {
 
+class FMarkupBuilder;
 class FYamlBuilder;
 struct FSchemaBatch;
 struct FEnumSchema;
@@ -36,20 +37,21 @@ using FYamlBuilderPtr = TUniquePtr<FYamlBuilder, FYamlBuilderDeleter>;
 class FBatchPrinter
 {
 public:
-	FBatchPrinter(FYamlBuilder& InTextBuilder, const FBatchIds& InIds);
+	FBatchPrinter(FYamlBuilder& InMarkupBuilder, const FBatchIds& InIds);
 	~FBatchPrinter();
 
 	void PrintSchemas();
 	void PrintObjects(TConstArrayView<FStructView> Structs);
 
 private:
+	FBatchPrinter(FMarkupBuilder& InMarkupBuilder, const FBatchIds& InIds);
 	void PrintStructSchema(const FStructSchema& Struct, FSchemaBatchId BatchId);
 	void PrintEnumSchema(const FEnumSchema& Enum);
 
 	template<typename IntType>
 	void PrintEnumConstants(TConstArrayView<FNameId> EnumNames, TConstArrayView<IntType> Constants, bool bFlagMode);
 
-	FYamlBuilder& TextBuilder;
+	FMarkupBuilder& MarkupBuilder;
 	const FBatchIds& Ids;
 };
 
diff --git a/Engine/Plugins/Experimental/PlainProps/Source/Private/PlainPropsPrint.cpp b/Engine/Plugins/Experimental/PlainProps/Source/Private/PlainPropsPrint.cpp
--- a/Engine/Plugins/Experimental/PlainProps/Source/Private/PlainPropsPrint.cpp
+++ b/Engine/Plugins/Experimental/PlainProps/Source/Private/PlainPropsPrint.cpp
@@ -751,8 +751,8 @@ private:
 class FMemberPrinter
 {
 public:
-	FMemberPrinter(FYamlBuilder& InTextBuilder, const FBatchIds& InIds)
-	: TextBuilder(InTextBuilder)
+	FMemberPrinter(FMarkupBuilder& InMarkupBuilder, const FBatchIds& InIds)
+	: MarkupBuilder(InMarkupBuilder)
 	, Ids(InIds)
 	{}
 
@@ -784,26 +784,31 @@ private:
 		if constexpr (PrintWithComments)
 		{
 			PrintSchema(/* out */ Tmp, Ids, Type, Schema);
-			TextBuilder.AddComment(Tmp);
+			MarkupBuilder.AddComment(Tmp);
 			Tmp.Reset();
 		}
 	}
 
-	FYamlBuilder& TextBuilder;
+	FMarkupBuilder& MarkupBuilder;
 	const FBatchIds& Ids;
 	TUtf8StringBuilder<256> Tmp;
 };
 
 ///////////////////////////////////////////////////////////////////////////////
 
+static void PrintBatch(FMarkupBuilder& MarkupBuilder, FBatchPrinter& Printer, TConstArrayView<FStructView> Objects)
+{
+	MarkupBuilder.BeginDocument();
+	Printer.PrintSchemas();
+	Printer.PrintObjects(Objects);
+	MarkupBuilder.EndDocument();
+}
+
 void PrintYamlBatch(FUtf8Builder& Out, const FBatchIds& Ids, TConstArrayView<FStructView> Objects)
 {
 	FYamlBuilder YamlBuilder(Out);
 	FBatchPrinter Printer(YamlBuilder, Ids);
-	YamlBuilder.BeginDocument();
-	Printer.PrintSchemas();
-	Printer.PrintObjects(Objects);
-	YamlBuilder.EndDocument();
+	PrintBatch(YamlBuilder, Printer, Objects);
 }
 
 ///////////////////////////////////////////////////////////////////////////////
@@ -1086,8 +1091,12 @@ void FYamlBuilder::PopScope()
 
 ///////////////////////////////////////////////////////////////////////////////
 
-FBatchPrinter::FBatchPrinter(FYamlBuilder& InTextBuilder, const FBatchIds& InIds)
-: TextBuilder(InTextBuilder)
+FBatchPrinter::FBatchPrinter(FYamlBuilder& InMarkupBuilder, const FBatchIds& InIds)
+: FBatchPrinter(static_cast<FMarkupBuilder&>(InMarkupBuilder), InIds)
+{}
+
+FBatchPrinter::FBatchPrinter(FMarkupBuilder& InMarkupBuilder, const FBatchIds& InIds)
+: MarkupBuilder(InMarkupBuilder)
 , Ids(InIds)
 {}
 
@@ -1096,29 +1105,29 @@ FBatchPrinter::~FBatchPrinter()
 
 void FBatchPrinter::PrintSchemas()
 {
-	TextBuilder.BeginStruct(GLiterals.Structs);
+	MarkupBuilder.BeginStruct(GLiterals.Structs);
 	for (const FStructSchema& Struct : GetStructSchemas(Ids.GetSchemas()))
 	{
 		PrintStructSchema(Struct, Ids.GetBatchId());
 	}
-	TextBuilder.EndStruct();
+	MarkupBuilder.EndStruct();
 
-	TextBuilder.BeginStruct(GLiterals.Enums);
+	MarkupBuilder.BeginStruct(GLiterals.Enums);
 	for (const FEnumSchema& EnumSchema : GetEnumSchemas(Ids.GetSchemas()))
 	{
 		PrintEnumSchema(EnumSchema);
 	}
-	TextBuilder.EndStruct();
+	MarkupBuilder.EndStruct();
 }
 
 void FBatchPrinter::PrintObjects(TConstArrayView<FStructView> Objects)
 {
-	TextBuilder.BeginRange(GLiterals.Objects);
+	MarkupBuilder.BeginRange(GLiterals.Objects);
 	for (FStructView Object : Objects)
 	{
-		FMemberPrinter(TextBuilder, Ids).PrintMembers(Object);
+		FMemberPrinter(MarkupBuilder, Ids).PrintMembers(Object);
 	}
-	TextBuilder.EndRange();
+	MarkupBuilder.EndRange();
 }
 
 static void PrintMemberSchema(FUtf8Builder& Out, const FBatchIds& Ids, const FMemberSchemaView& Schema)
@@ -1145,43 +1154,43 @@ void FBatchPrinter::PrintStructSchema(const FStructSchema& Struct, FSchemaBatchI
 {
 	FStructSchemaReader Reader(Struct, BatchId);
 
-	TextBuilder.BeginStruct(*FPrintId<128>(Ids, Reader.GetStruct()));
+	MarkupBuilder.BeginStruct(*FPrintId<128>(Ids, Reader.GetStruct()));
 
 	if (uint16 Version = Reader.GetVersion())
 	{
-		TextBuilder.AddLeaf(GLiterals.Version, Version);
+		MarkupBuilder.AddLeaf(GLiterals.Version, Version);
 	}
 
 	if (Reader.HasSuper())
 	{
 		const FStructSchema& SuperSchema = Reader.GetSuper().Resolve();
-		TextBuilder.AddLeaf(GLiterals.DeclaredSuper, *FPrintId<128>(Ids, SuperSchema.Type));
+		MarkupBuilder.AddLeaf(GLiterals.DeclaredSuper, *FPrintId<128>(Ids, SuperSchema.Type));
 	}
 
-	TextBuilder.BeginStruct(GLiterals.Members);
+	MarkupBuilder.BeginStruct(GLiterals.Members);
 	TUtf8StringBuilder<256> Buf;
 	while (Reader.HasMore())
 	{
 		Ids.AppendString(Buf, Reader.PeekName());
-		TextBuilder.AddId(Buf);
+		MarkupBuilder.AddId(Buf);
 		Buf.Reset();
 
 		PrintMemberSchema(Buf, Ids, Reader.GrabMember());
-		TextBuilder.AddValue(FUtf8StringView(Buf));
+		MarkupBuilder.AddValue(FUtf8StringView(Buf));
 		Buf.Reset();
 	}
-	TextBuilder.EndStruct();
+	MarkupBuilder.EndStruct();
 
-	TextBuilder.EndStruct();
+	MarkupBuilder.EndStruct();
 }
 
 void FBatchPrinter::PrintEnumSchema(const FEnumSchema& Enum)
 {
-	TextBuilder.BeginStruct(*FPrintId<128>(Ids, Enum.Type));
-	TextBuilder.AddLeaf(GLiterals.FlagMode, !!Enum.FlagMode);
-	TextBuilder.AddLeaf(GLiterals.Width, Enum.Width);
+	MarkupBuilder.BeginStruct(*FPrintId<128>(Ids, Enum.Type));
+	MarkupBuilder.AddLeaf(GLiterals.FlagMode, !!Enum.FlagMode);
+	MarkupBuilder.AddLeaf(GLiterals.Width, Enum.Width);
 
-	TextBuilder.BeginStruct(GLiterals.Constants);
+	MarkupBuilder.BeginStruct(GLiterals.Constants);
 	TConstArrayView<FNameId> EnumNames = MakeConstArrayView(Enum.Footer, Enum.Num);
 	switch (Enum.Width)
 	{
@@ -1198,9 +1207,9 @@ void FBatchPrinter::PrintEnumSchema(const FEnumSchema& Enum)
 			PrintEnumConstants(EnumNames, GetConstants<uint64>(Enum), Enum.FlagMode);
 			break;
 	}
-	TextBuilder.EndStruct();
+	MarkupBuilder.EndStruct();
 
-	TextBuilder.EndStruct();
+	MarkupBuilder.EndStruct();
 }
 
 template<typename IntType>
@@ -1215,7 +1224,7 @@ void FBatchPrinter::PrintEnumConstants(
 		check(EnumNames.Num() == Constants.Num());
 		for (uint16 Idx = 0; Idx < NamesNum; ++Idx)
 		{
-			TextBuilder.AddLeaf(*FPrintId<128>(Ids, EnumNames[Idx]), (uint64)Constants[Idx]);
+			MarkupBuilder.AddLeaf(*FPrintId<128>(Ids, EnumNames[Idx]), (uint64)Constants[Idx]);
 		}
 	}
 	else if (bFlagMode)
@@ -1223,7 +1232,7 @@ void FBatchPrinter::PrintEnumConstants(
 		uint64 Value = 1;
 		for (uint16 Idx = 0; Idx < NamesNum; ++Idx)
 		{
-			TextBuilder.AddLeaf(*FPrintId<128>(Ids, EnumNames[Idx]), Value);
+			MarkupBuilder.AddLeaf(*FPrintId<128>(Ids, EnumNames[Idx]), Value);
 			Value <<= 1;
 		}
 	}
@@ -1231,7 +1240,7 @@ void FBatchPrinter::PrintEnumConstants(
 	{
 		for (uint16 Idx = 0; Idx < NamesNum; ++Idx)
 		{
-			TextBuilder.AddLeaf(*FPrintId<128>(Ids, EnumNames[Idx]), (uint64)Idx);
+			MarkupBuilder.AddLeaf(*FPrintId<128>(Ids, EnumNames[Idx]), (uint64)Idx);
 		}
 	}
 }
@@ -1241,57 +1250,57 @@ void FBatchPrinter::PrintEnumConstants(
 void FMemberPrinter::PrintMembers(FStructView StructView)
 {
 	const FStructSchema& Schema = StructView.Schema.Resolve();
-	TextBuilder.BeginStruct();
-	TextBuilder.BeginStruct(*FPrintId<128>(Ids, Schema.Type));
+	MarkupBuilder.BeginStruct();
+	MarkupBuilder.BeginStruct(*FPrintId<128>(Ids, Schema.Type));
 	PrintMembersInternal({ EMemberKind::Struct }, StructView, Schema);
-	TextBuilder.EndStruct();
+	MarkupBuilder.EndStruct();
 }
 
 void FMemberPrinter::PrintLeaf(FMemberId Id, FLeafView LeafView)
 {
-	TextBuilder.AddId(*FPrintId<128>(Ids, Id));
+	MarkupBuilder.AddId(*FPrintId<128>(Ids, Id));
 
 	switch (LeafView.Leaf.Type)
 	{
 	case ELeafType::Bool:
-		TextBuilder.AddValue(LeafView.AsBool());
+		MarkupBuilder.AddValue(LeafView.AsBool());
 		break;
 	case ELeafType::IntS:
 		switch (LeafView.Leaf.Width)
 		{
-			case ELeafWidth::B8:	TextBuilder.AddValue(LeafView.AsS8()); break;
-			case ELeafWidth::B16:	TextBuilder.AddValue(LeafView.AsS16()); break;
-			case ELeafWidth::B32:	TextBuilder.AddValue(LeafView.AsS32()); break;
-			case ELeafWidth::B64:	TextBuilder.AddValue(LeafView.AsS64()); break;
+			case ELeafWidth::B8:	MarkupBuilder.AddValue(LeafView.AsS8()); break;
+			case ELeafWidth::B16:	MarkupBuilder.AddValue(LeafView.AsS16()); break;
+			case ELeafWidth::B32:	MarkupBuilder.AddValue(LeafView.AsS32()); break;
+			case ELeafWidth::B64:	MarkupBuilder.AddValue(LeafView.AsS64()); break;
 		}
 		break;
 	case ELeafType::IntU:
 		switch (LeafView.Leaf.Width)
 		{
-			case ELeafWidth::B8:	TextBuilder.AddValue(LeafView.AsU8()); break;
-			case ELeafWidth::B16:	TextBuilder.AddValue(LeafView.AsU16()); break;
-			case ELeafWidth::B32:	TextBuilder.AddValue(LeafView.AsU32()); break;
-			case ELeafWidth::B64:	TextBuilder.AddValue(LeafView.AsU64()); break;
+			case ELeafWidth::B8:	MarkupBuilder.AddValue(LeafView.AsU8()); break;
+			case ELeafWidth::B16:	MarkupBuilder.AddValue(LeafView.AsU16()); break;
+			case ELeafWidth::B32:	MarkupBuilder.AddValue(LeafView.AsU32()); break;
+			case ELeafWidth::B64:	MarkupBuilder.AddValue(LeafView.AsU64()); break;
 		}
 		break;
 	case ELeafType::Float:
 		if (LeafView.Leaf.Width == ELeafWidth::B32)
 		{
-			TextBuilder.AddValue(LeafView.AsFloat());
+			MarkupBuilder.AddValue(LeafView.AsFloat());
 		}
 		else
 		{
 			check(LeafView.Leaf.Width == ELeafWidth::B64);
-			TextBuilder.AddValue(LeafView.AsDouble());
+			MarkupBuilder.AddValue(LeafView.AsDouble());
 		}
 		break;
 	case ELeafType::Hex:
 		switch (LeafView.Leaf.Width)
 		{
-			case ELeafWidth::B8:	TextBuilder.AddValue(FPrintFormatHex<uint8 >{LeafView.AsHex8()});  break;
-			case ELeafWidth::B16:	TextBuilder.AddValue(FPrintFormatHex<uint16>{LeafView.AsHex16()}); break;
-			case ELeafWidth::B32:	TextBuilder.AddValue(FPrintFormatHex<uint32>{LeafView.AsHex32()}); break;
-			case ELeafWidth::B64:	TextBuilder.AddValue(FPrintFormatHex<uint64>{LeafView.AsHex64()}); break;
+			case ELeafWidth::B8:	MarkupBuilder.AddValue(FPrintFormatHex<uint8 >{LeafView.AsHex8()});  break;
+			case ELeafWidth::B16:	MarkupBuilder.AddValue(FPrintFormatHex<uint16>{LeafView.AsHex16()}); break;
+			case ELeafWidth::B32:	MarkupBuilder.AddValue(FPrintFormatHex<uint32>{LeafView.AsHex32()}); break;
+			case ELeafWidth::B64:	MarkupBuilder.AddValue(FPrintFormatHex<uint64>{LeafView.AsHex64()}); break;
 		}
 		break;
 	case ELeafType::Enum:
@@ -1306,9 +1315,9 @@ void FMemberPrinter::PrintLeaf(FMemberId Id, FLeafView LeafView)
 	case ELeafType::Unicode:
 		switch (LeafView.Leaf.Width)
 		{
-			case ELeafWidth::B8:	TextBuilder.AddValue(LeafView.AsChar8());	break;
-			case ELeafWidth::B16:	TextBuilder.AddValue(LeafView.AsChar16());	break;
-			case ELeafWidth::B32:	TextBuilder.AddValue(LeafView.AsChar32());	break;
+			case ELeafWidth::B8:	MarkupBuilder.AddValue(LeafView.AsChar8());	break;
+			case ELeafWidth::B16:	MarkupBuilder.AddValue(LeafView.AsChar16());	break;
+			case ELeafWidth::B32:	MarkupBuilder.AddValue(LeafView.AsChar32());	break;
 			case ELeafWidth::B64:	check(false);								break;
 		};
 	}
@@ -1330,12 +1339,12 @@ void FMemberPrinter::PrintLeafEnum(FMemberId Id, FLeafView LeafView)
 	TUtf8StringBuilder<64> Buffer;
 	PrintEnumValue(Buffer, Value, EnumSchema, Ids);
 
-	TextBuilder.AddValue(FUtf8StringView(Buffer));
+	MarkupBuilder.AddValue(FUtf8StringView(Buffer));
 }
 
 void FMemberPrinter::PrintStruct(FOptionalMemberId MemberId, FStructType StructType, FStructView StructView)
 {
-	TextBuilder.BeginStruct(*FPrintId<128>(Ids, MemberId));
+	MarkupBuilder.BeginStruct(*FPrintId<128>(Ids, MemberId));
 	PrintMembersInternal(StructType, StructView, StructView.Schema.Resolve());
 }
 
@@ -1347,7 +1356,7 @@ void FMemberPrinter::PrintRange(FMemberId Id, FRangeType RangeType, const FRange
 	}
 	else
 	{
-		TextBuilder.BeginRange(*FPrintId<128>(Ids, Id));
+		MarkupBuilder.BeginRange(*FPrintId<128>(Ids, Id));
 		PrintRangeInternal(RangeType, RangeView);
 	}
 }
@@ -1357,43 +1366,43 @@ void FMemberPrinter::PrintLeaves(FUnpackedLeafType Leaf, const FLeafRangeView& L
 	switch (Leaf.Type)
 	{
 	case ELeafType::Bool:
-		for (bool b : LeafRange.AsBools()) { TextBuilder.AddItem(b); } break;
+		for (bool b : LeafRange.AsBools()) { MarkupBuilder.AddItem(b); } break;
 	case ELeafType::IntS:
 		switch (Leaf.Width)
 		{
-			case ELeafWidth::B8:  for (const int8&  I : LeafRange.AsS8s())  { TextBuilder.AddItem(I); } break;
-			case ELeafWidth::B16: for (const int16& I : LeafRange.AsS16s()) { TextBuilder.AddItem(I); } break;
-			case ELeafWidth::B32: for (const int32& I : LeafRange.AsS32s()) { TextBuilder.AddItem(I); } break;
-			case ELeafWidth::B64: for (const int64& I : LeafRange.AsS64s()) { TextBuilder.AddItem(I); } break;
+			case ELeafWidth::B8:  for (const int8&  I : LeafRange.AsS8s())  { MarkupBuilder.AddItem(I); } break;
+			case ELeafWidth::B16: for (const int16& I : LeafRange.AsS16s()) { MarkupBuilder.AddItem(I); } break;
+			case ELeafWidth::B32: for (const int32& I : LeafRange.AsS32s()) { MarkupBuilder.AddItem(I); } break;
+			case ELeafWidth::B64: for (const int64& I : LeafRange.AsS64s()) { MarkupBuilder.AddItem(I); } break;
 		}
 		break;
 	case ELeafType::IntU:
 		switch (Leaf.Width)
 		{
-			case ELeafWidth::B8:  for (const uint8&  U : LeafRange.AsU8s())  { TextBuilder.AddItem(U); } break;
-			case ELeafWidth::B16: for (const uint16& U : LeafRange.AsU16s()) { TextBuilder.AddItem(U); } break;
-			case ELeafWidth::B32: for (const uint32& U : LeafRange.AsU32s()) { TextBuilder.AddItem(U); } break;
-			case ELeafWidth::B64: for (const uint64& U : LeafRange.AsU64s()) { TextBuilder.AddItem(U); } break;
+			case ELeafWidth::B8:  for (const uint8&  U : LeafRange.AsU8s())  { MarkupBuilder.AddItem(U); } break;
+			case ELeafWidth::B16: for (const uint16& U : LeafRange.AsU16s()) { MarkupBuilder.AddItem(U); } break;
+			case ELeafWidth::B32: for (const uint32& U : LeafRange.AsU32s()) { MarkupBuilder.AddItem(U); } break;
+			case ELeafWidth::B64: for (const uint64& U : LeafRange.AsU64s()) { MarkupBuilder.AddItem(U); } break;
 		}
 		break;
 	case ELeafType::Float:
 		if (Leaf.Width == ELeafWidth::B32)
 		{
-			for (const float& f : LeafRange.AsFloats()) { TextBuilder.AddItem(f); }
+			for (const float& f : LeafRange.AsFloats()) { MarkupBuilder.AddItem(f); }
 		}
 		else
 		{
 			check(Leaf.Width == ELeafWidth::B64);
-			for (const double& d : LeafRange.AsDoubles()) { TextBuilder.AddItem(d); }
+			for (const double& d : LeafRange.AsDoubles()) { MarkupBuilder.AddItem(d); }
 		}
 		break;
 	case ELeafType::Hex:
 		switch (Leaf.Width)
 		{
-			case ELeafWidth::B8:  for (uint8  U : LeafRange.AsHex8s())  { TextBuilder.AddItem(FPrintFormatHex<uint8> {U}); } break;
-			case ELeafWidth::B16: for (uint16 U : LeafRange.AsHex16s()) { TextBuilder.AddItem(FPrintFormatHex<uint16>{U}); } break;
-			case ELeafWidth::B32: for (uint32 U : LeafRange.AsHex32s()) { TextBuilder.AddItem(FPrintFormatHex<uint32>{U}); } break;
-			case ELeafWidth::B64: for (uint64 U : LeafRange.AsHex64s()) { TextBuilder.AddItem(FPrintFormatHex<uint64>{U}); } break;
+			case ELeafWidth::B8:  for (uint8  U : LeafRange.AsHex8s())  { MarkupBuilder.AddItem(FPrintFormatHex<uint8> {U}); } break;
+			case ELeafWidth::B16: for (uint16 U : LeafRange.AsHex16s()) { MarkupBuilder.AddItem(FPrintFormatHex<uint16>{U}); } break;
+			case ELeafWidth::B32: for (uint32 U : LeafRange.AsHex32s()) { MarkupBuilder.AddItem(FPrintFormatHex<uint32>{U}); } break;
+			case ELeafWidth::B64: for (uint64 U : LeafRange.AsHex64s()) { MarkupBuilder.AddItem(FPrintFormatHex<uint64>{U}); } break;
 		}
 		break;
 	case ELeafType::Enum:
@@ -1422,7 +1431,7 @@ void FMemberPrinter::PrintLeavesEnum(FUnpackedLeafType Leaf, const FLeafRangeVie
 	{
 		TUtf8StringBuilder<128> Buffer;
 		PrintEnumValue(Buffer, Value, EnumSchema, Ids);
-		TextBuilder.AddItem(FUtf8StringView(Buffer));
+		MarkupBuilder.AddItem(FUtf8StringView(Buffer));
 	}
 }
 
@@ -1430,7 +1439,7 @@ void FMemberPrinter::PrintStructs(FStructType StructType, const FStructRangeView
 {
 	for (FStructView StructView : StructRange)
 	{
-		TextBuilder.BeginStruct();
+		MarkupBuilder.BeginStruct();
 		PrintMembersInternal(StructType, StructView, StructView.Schema.Resolve());
 	}
 }
@@ -1445,7 +1454,7 @@ void FMemberPrinter::PrintRanges(FRangeType RangeType, const FNestedRangeView& N
 		}
 		else
 		{
-			TextBuilder.BeginRange();
+			MarkupBuilder.BeginRange();
 			PrintRangeInternal(RangeType, RangeView);
 		}
 	}
@@ -1463,7 +1472,7 @@ void FMemberPrinter::PrintMembersInternal(FStructType StructType, FStructView St
 
 	if (StructType.IsDynamic)
 	{
-		TextBuilder.AddLeaf(GLiterals.Dynamic, *FPrintId<128>(Ids, Schema.Type));
+		MarkupBuilder.AddLeaf(GLiterals.Dynamic, *FPrintId<128>(Ids, Schema.Type));
 	}
 	while (It.HasMore())
 	{
@@ -1482,7 +1491,7 @@ void FMemberPrinter::PrintMembersInternal(FStructType StructType, FStructView St
 				break;
 		}
 	}
-	TextBuilder.EndStruct();
+	MarkupBuilder.EndStruct();
 
 	if (!HasMembers)
 	{
@@ -1510,7 +1519,7 @@ void FMemberPrinter::PrintRangeInternal(FRangeType RangeType, const FRangeView&
 			PrintRanges(Schema.ItemType.AsRange(), RangeView.AsRanges());
 			break;
 	}
-	TextBuilder.EndRange();
+	MarkupBuilder.EndRange();
 
 	if (RangeView.Num() == 0)
 	{
@@ -1525,7 +1534,7 @@ bool FMemberPrinter::IsUnicodeString(const FRangeView& RangeView)
 }
 
 template <typename CharType, typename RangeCharType>
-static void AddUnicodeRangeLeaf(FYamlBuilder& TextBuilder, FOptionalMemberId Id, TRangeView<RangeCharType> Range)
+static void AddUnicodeRangeLeaf(FMarkupBuilder& MarkupBuilder, FOptionalMemberId Id, TRangeView<RangeCharType> Range)
 {
 	check(Range.Num() > 0);
 	const CharType* Src = reinterpret_cast<const CharType*>(Range.begin());
@@ -1539,11 +1548,11 @@ static void AddUnicodeRangeLeaf(FYamlBuilder& TextBuilder, FOptionalMemberId Id,
 	check(DstEnd - Dst == DstLen);
 	if (Id)
 	{
-		TextBuilder.AddValue(FUtf8StringView(Dst, DstLen));
+		MarkupBuilder.AddValue(FUtf8StringView(Dst, DstLen));
 	}
 	else
 	{
-		TextBuilder.AddItem(FUtf8StringView(Dst, DstLen));
+		MarkupBuilder.AddItem(FUtf8StringView(Dst, DstLen));
 	}
 }
 
@@ -1553,7 +1562,7 @@ void FMemberPrinter::PrintUnicodeRangeAsLeaf(FOptionalMemberId Id, FRangeType Ra
 
 	if (Id)
 	{
-		TextBuilder.AddId(*FPrintId<128>(Ids, Id));
+		MarkupBuilder.AddId(*FPrintId<128>(Ids, Id));
 	}
 
 	const FLeafRangeView LeafRange = RangeView.AsLeaves();
@@ -1561,9 +1570,9 @@ void FMemberPrinter::PrintUnicodeRangeAsLeaf(FOptionalMemberId Id, FRangeType Ra
 
 	switch (Leaf.Width)
 	{
-		case ELeafWidth::B8:	AddUnicodeRangeLeaf<UTF8CHAR >(TextBuilder, Id, LeafRange.AsUtf8());	break;
-		case ELeafWidth::B16:	AddUnicodeRangeLeaf<UTF16CHAR>(TextBuilder, Id, LeafRange.AsUtf16());	break;
-		case ELeafWidth::B32:	AddUnicodeRangeLeaf<UTF32CHAR>(TextBuilder, Id, LeafRange.AsUtf32());	break;
+		case ELeafWidth::B8:	AddUnicodeRangeLeaf<UTF8CHAR >(MarkupBuilder, Id, LeafRange.AsUtf8());	break;
+		case ELeafWidth::B16:	AddUnicodeRangeLeaf<UTF16CHAR>(MarkupBuilder, Id, LeafRange.AsUtf16());	break;
+		case ELeafWidth::B32:	AddUnicodeRangeLeaf<UTF32CHAR>(MarkupBuilder, Id, LeafRange.AsUtf32());	break;
 		case ELeafWidth::B64:	check(false);															break;
 	};
 
