nss: Fixed OOB write.

#rnx

#Synced-CL 49920676

diff --git a/Engine/Plugins/Experimental/PlainProps/Source/Private/PlainPropsPrint.cpp b/Engine/Plugins/Experimental/PlainProps/Source/Private/PlainPropsPrint.cpp
--- a/Engine/Plugins/Experimental/PlainProps/Source/Private/PlainPropsPrint.cpp
+++ b/Engine/Plugins/Experimental/PlainProps/Source/Private/PlainPropsPrint.cpp
@@ -1541,7 +1541,7 @@ static void AddUnicodeRangeLeaf(FMarkupBuilder& MarkupBuilder, FOptionalMemberId
 	const int32 SrcLen = IntCastChecked<int32>(Range.Num());
 	const int32 DstLen = FPlatformString::ConvertedLength<UTF8CHAR>(Src, SrcLen);
 	TArray<UTF8CHAR, TInlineAllocator<1024>> Buf;
-	Buf.Reserve(DstLen);
+	Buf.SetNumUninitialized(DstLen);
 	UTF8CHAR* Dst = Buf.GetData();
 	const UTF8CHAR* DstEnd = FPlatformString::Convert(Dst, DstLen, Src, SrcLen);
 	check(DstEnd);
