nss: Correct double-quotes for JSON leaf IDs

#rnx

#Synced-CL 49920676

diff --git a/Engine/Plugins/Experimental/PlainProps/Source/Private/PlainPropsPrint.cpp b/Engine/Plugins/Experimental/PlainProps/Source/Private/PlainPropsPrint.cpp
--- a/Engine/Plugins/Experimental/PlainProps/Source/Private/PlainPropsPrint.cpp
+++ b/Engine/Plugins/Experimental/PlainProps/Source/Private/PlainPropsPrint.cpp
@@ -680,7 +680,7 @@ void FJsonBuilder::WriteId(FUtf8StringView Id)
 {
 	Prologue(EAction::Add);
 
-	PrintQuotedString(Builder, Id);
+	PrintQuotedString(Builder, Id, '\"');
 	Builder << " : ";
 }
 
