PlainProps: Use TCustomSpecifier::Super when constructing the FStructSpec in TScopedStructBinding.
Change TCustomSpecifier::Super type to FOptionalDeclId.
Check that the last member id has been initialized.

#rnx

#Synced-CL 50399333

diff --git a/Engine/Plugins/Experimental/PlainProps/Source/Public/PlainPropsBindCtti.h b/Engine/Plugins/Experimental/PlainProps/Source/Public/PlainPropsBindCtti.h
--- a/Engine/Plugins/Experimental/PlainProps/Source/Public/PlainPropsBindCtti.h
+++ b/Engine/Plugins/Experimental/PlainProps/Source/Public/PlainPropsBindCtti.h
@@ -259,7 +259,7 @@ struct TCustomInit : FCustomInit, TInit<Ids>
 template<class Ids, uint32 N>
 struct TCustomSpecifier : TCustomInit<Ids>
 {
-	FOptionalStructId				Super;
+	FOptionalDeclId					Super;
 	uint16							Version = 0;
 	FMemberSpec						Members[N];
 
@@ -314,7 +314,8 @@ struct TScopedStructBinding : FBothStructId, CustomBinding
 	: FBothStructId(Both)
 	, CustomBinding(/* out */ Init)
 	{
-		FStructSpec Decl = { DeclId, /* todo */ NoId, Init.Version, TOccupancyOf<T>::Value, CustomBinding::MemberIds, Init.Members };
+		checkf(N > 0 && CustomBinding::MemberIds[N - 1] != FMemberId{}, TEXT("MemberIds not fully initialized"));
+		FStructSpec Decl = { DeclId, Init.Super, Init.Version, TOccupancyOf<T>::Value, CustomBinding::MemberIds, Init.Members };
 		Runtime::GetCustoms().BindStruct(BindId, *this, Decl, Init.LoweredInners);
 	}
 
