nss: generalise to the markup base class where possible

#rnx

#Synced-CL 50479750

diff --git a/Engine/Plugins/Experimental/PlainProps/Source/Private/PlainPropsInternalPrint.h b/Engine/Plugins/Experimental/PlainProps/Source/Private/PlainPropsInternalPrint.h
--- a/Engine/Plugins/Experimental/PlainProps/Source/Private/PlainPropsInternalPrint.h
+++ b/Engine/Plugins/Experimental/PlainProps/Source/Private/PlainPropsInternalPrint.h
@@ -12,7 +12,6 @@ namespace PlainProps
 {
 
 class FMarkupBuilder;
-class FYamlBuilder;
 struct FSchemaBatch;
 struct FEnumSchema;
 struct FStructSchema;
@@ -24,11 +23,12 @@ inline TStringBuilderWithBuffer<TCHAR, 128> Print(FUtf8StringView View)
 
 ///////////////////////////////////////////////////////////////////////////////
 
-struct FYamlBuilderDeleter
+struct FMarkupBuilderDeleter
 {
-	PLAINPROPS_API void operator()(FYamlBuilder* Ptr) const;
+	PLAINPROPS_API void operator()(FMarkupBuilder* Ptr) const;
 };
-using FYamlBuilderPtr = TUniquePtr<FYamlBuilder, FYamlBuilderDeleter>;
+using FYamlBuilderPtr = TUniquePtr<FMarkupBuilder, FMarkupBuilderDeleter>;
+using FJsonBuilderPtr = TUniquePtr<FMarkupBuilder, FMarkupBuilderDeleter>;
 
 [[nodiscard]] FYamlBuilderPtr MakeYamlBuilder(FUtf8StringBuilderBase& StringBuilder);
 
@@ -37,14 +37,13 @@ using FYamlBuilderPtr = TUniquePtr<FYamlBuilder, FYamlBuilderDeleter>;
 class FBatchPrinter
 {
 public:
-	FBatchPrinter(FYamlBuilder& InMarkupBuilder, const FBatchIds& InIds);
+	FBatchPrinter(FMarkupBuilder& InMarkupBuilder, const FBatchIds& InIds);
 	~FBatchPrinter();
 
 	void PrintSchemas();
 	void PrintObjects(TConstArrayView<FStructView> Structs);
 
 private:
-	FBatchPrinter(FMarkupBuilder& InMarkupBuilder, const FBatchIds& InIds);
 	void PrintStructSchema(const FStructSchema& Struct, FSchemaBatchId BatchId);
 	void PrintEnumSchema(const FEnumSchema& Enum);
 
diff --git a/Engine/Plugins/Experimental/PlainProps/Source/Private/PlainPropsPrint.cpp b/Engine/Plugins/Experimental/PlainProps/Source/Private/PlainPropsPrint.cpp
--- a/Engine/Plugins/Experimental/PlainProps/Source/Private/PlainPropsPrint.cpp
+++ b/Engine/Plugins/Experimental/PlainProps/Source/Private/PlainPropsPrint.cpp
@@ -391,6 +391,7 @@ class FMarkupBuilder
 	: public IMarkupBuilder
 {
 public:
+	virtual		~FMarkupBuilder()				= default;
 	void		BeginDocument()					{ PushDocument(); }
 	void		EndDocument()					{ PopDocument(); }
 	void		BeginStruct()					{ PushStruct(); }
@@ -931,9 +932,9 @@ const FMemberId* FStructSchemaReader::GetMemberNames() const
 
 ///////////////////////////////////////////////////////////////////////////////
 
-void FYamlBuilderDeleter::operator()(FYamlBuilder* YamlBuilder) const
+void FMarkupBuilderDeleter::operator()(FMarkupBuilder* MarkupBuilder) const
 {
-	delete YamlBuilder;
+	delete MarkupBuilder;
 }
 
 FYamlBuilderPtr MakeYamlBuilder(FUtf8Builder& StringBuilder)
@@ -1091,10 +1092,6 @@ void FYamlBuilder::PopScope()
 
 ///////////////////////////////////////////////////////////////////////////////
 
-FBatchPrinter::FBatchPrinter(FYamlBuilder& InMarkupBuilder, const FBatchIds& InIds)
-: FBatchPrinter(static_cast<FMarkupBuilder&>(InMarkupBuilder), InIds)
-{}
-
 FBatchPrinter::FBatchPrinter(FMarkupBuilder& InMarkupBuilder, const FBatchIds& InIds)
 : MarkupBuilder(InMarkupBuilder)
 , Ids(InIds)
