nss: Base class for common markup tokenisation.

#rnx

#Synced-CL 50479750

diff --git a/Engine/Plugins/Experimental/PlainProps/Source/Private/PlainPropsInternalParse.h b/Engine/Plugins/Experimental/PlainProps/Source/Private/PlainPropsInternalParse.h
--- a/Engine/Plugins/Experimental/PlainProps/Source/Private/PlainPropsInternalParse.h
+++ b/Engine/Plugins/Experimental/PlainProps/Source/Private/PlainPropsInternalParse.h
@@ -18,7 +18,7 @@ namespace PlainProps
 {
 
 class FScratchAllocator;
-class FYamlTokenizer;
+class FMarkupTokenizer;
 struct FBuiltStruct;
 struct FEnumerator;
 
@@ -188,7 +188,7 @@ struct FParseDeclarations final : IDeclarations
 class FBatchParser
 {
 public:
-	FBatchParser(FYamlTokenizer& InTokenizer, FScratchAllocator& InScratch)
+	FBatchParser(FMarkupTokenizer& InTokenizer, FScratchAllocator& InScratch)
 	: Tokenizer(InTokenizer), Scratch(InScratch), Types(FDebugIds(Names))
 	{}
 
@@ -202,7 +202,7 @@ private:
 	void ParseObjects(TConstArrayView64<FToken> TokensView);
 	void Write(TArray64<uint8>& Out);
 
-	FYamlTokenizer&									Tokenizer;
+	FMarkupTokenizer&								Tokenizer;
 	FScratchAllocator&								Scratch;
 	FTextIndexer									Names;
 	FParseDeclarations								Types;
diff --git a/Engine/Plugins/Experimental/PlainProps/Source/Private/PlainPropsParse.cpp b/Engine/Plugins/Experimental/PlainProps/Source/Private/PlainPropsParse.cpp
--- a/Engine/Plugins/Experimental/PlainProps/Source/Private/PlainPropsParse.cpp
+++ b/Engine/Plugins/Experimental/PlainProps/Source/Private/PlainPropsParse.cpp
@@ -510,17 +510,85 @@ struct FLineInfo
 	bool					EndOfFile		= false;
 };
 
+class FMarkupTokenizer
+{
+public:
+	virtual						~FMarkupTokenizer() = default;
+	virtual bool				HasMore() const = 0;
+	virtual TOptional<FToken>	GrabToken() = 0;
+
+	// PP-TEXT: GetLastError is a very temporary API while errors are unexpected
+	FUtf8StringView GetLastError() const;
+
+protected:
+							FMarkupTokenizer(FUtf8StringView InStringView);
+	FLineInfo				ReadOneLine();
+	void					SetLastError(FUtf8StringView Message, FUtf8StringView Context);
+
+private:
+	FUtf8StringView			ReadView;
+	TUtf8StringBuilder<256>	LastError;
+	uint32					LineNumber = 0;
+};
+
+FMarkupTokenizer::FMarkupTokenizer(FUtf8StringView InStringView)
+: ReadView(InStringView)
+{
+}
+
+FUtf8StringView FMarkupTokenizer::GetLastError() const
+{
+	return LastError;
+}
+
+FLineInfo FMarkupTokenizer::ReadOneLine()
+{
+	if (ReadView.IsEmpty())
+	{
+		return {
+			.View = ReadView,
+			.FullView = ReadView,
+			.LineNumber = LineNumber,
+			.EndOfFile = true
+		};
+	}
+
+	const FUtf8StringView LeadingSpaces = FAsciiSet::FindPrefixWith(ReadView, WhitespaceSet);
+	ReadView.RightChopInline(LeadingSpaces.Len());
+
+	const FUtf8StringView Line = FAsciiSet::FindPrefixWithout(ReadView, LinebreakSet);
+	ReadView.RightChopInline(Line.Len());
+
+	const FUtf8StringView TrailingLineBreaks = FAsciiSet::FindPrefixWith(ReadView, LinebreakSet);
+	ReadView.RightChopInline(TrailingLineBreaks.Len());
+
+	return { Line, Line, ++LineNumber, LeadingSpaces.Len() };
+}
+
+void FMarkupTokenizer::SetLastError(FUtf8StringView Message, FUtf8StringView Context)
+{
+	LastError.Reset();
+	if (Message.Len())
+	{
+		LastError.Appendf("'%.*s' in line %u: '%.*s'",
+			Message.Len(), Message.GetData(),
+			LineNumber,
+			Context.Len(),
+			Context.GetData());
+	}
+}
+
+////////////////////////////////////////////////////////////////////////////////
+
 class FYamlTokenizer
+	: public FMarkupTokenizer
 {
 public:
 	FYamlTokenizer(FUtf8StringView InStringView);
 	~FYamlTokenizer();
 
-	bool HasMore() const { return ReadLine.EndOfFile == false || ScopeStack.Num() > 0; }
-	TOptional<FToken> GrabToken();
-
-	// PP-TEXT: GetLastError is a very temporary API while errors are unexpected
-	FUtf8StringView GetLastError() const;
+	virtual bool HasMore() const override { return ReadLine.EndOfFile == false || ScopeStack.Num() > 0; }
+	virtual TOptional<FToken> GrabToken() override;
 
 private:
 	enum class EScope : uint8
@@ -549,19 +617,12 @@ private:
 	void						PushScopeStack(EScope Scope);
 	FToken						PopScopeStack();
 
-	FLineInfo					ReadOneLine();
 	void						AdvanceReadLines();
 
-	void						SetLastError(FUtf8StringView Message);
-
-	FUtf8StringView								Text;
-	FUtf8StringView								ReadView;
-	TUtf8StringBuilder<256>						LastError;
 	TArray<FScopeInfo, TInlineAllocator<32>>	ScopeStack;
 	TOptional<FToken>							CachedToken;
 	FLineInfo									ReadLine;
 	FLineInfo									NextLine;
-	uint32										LineNumber = 1;
 };
 
 ////////////////////////////////////////////////////////////////////////////////
@@ -725,8 +786,7 @@ private:
 ///////////////////////////////////////////////////////////////////////////////
 
 FYamlTokenizer::FYamlTokenizer(FUtf8StringView InStringView)
-: Text(InStringView)
-, ReadView(InStringView)
+: FMarkupTokenizer(InStringView)
 {
 	AdvanceReadLines();
 	AdvanceReadLines();
@@ -769,11 +829,6 @@ TOptional<FToken> FYamlTokenizer::GrabToken()
 	return Result;
 }
 
-FUtf8StringView FYamlTokenizer::GetLastError() const
-{
-	return LastError;
-}
-
 TOptional<FToken> FYamlTokenizer::GrabBeginDocument()
 {
 	if (ReadLine.EndOfFile)
@@ -944,7 +999,7 @@ TOptional<FToken> FYamlTokenizer::GrabTokenInternal()
 
 	check(ParsedToken.IsSet() == Error.IsEmpty());
 
-	SetLastError(Error);
+	SetLastError(Error, ReadLine.View);
 	return ParsedToken;
 }
 
@@ -959,25 +1014,6 @@ FToken FYamlTokenizer::PopScopeStack()
 	return { ScopeEndTokens[static_cast<uint8>(Last.Scope)], IntCastChecked<uint16>(ScopeStack.Num()) };
 }
 
-FLineInfo FYamlTokenizer::ReadOneLine()
-{
-	if (ReadView.IsEmpty())
-	{
-		return { .View = ReadView, .FullView = ReadView, .LineNumber = LineNumber, .EndOfFile = true };
-	}
-
-	const FUtf8StringView LeadingSpaces = FAsciiSet::FindPrefixWith(ReadView, WhitespaceSet);
-	ReadView.RightChopInline(LeadingSpaces.Len());
-
-	const FUtf8StringView Line = FAsciiSet::FindPrefixWithout(ReadView, LinebreakSet);
-	ReadView.RightChopInline(Line.Len());
-
-	const FUtf8StringView TrailingLineBreaks = FAsciiSet::FindPrefixWith(ReadView, LinebreakSet);
-	ReadView.RightChopInline(TrailingLineBreaks.Len());
-
-	return { Line, Line, LineNumber++, LeadingSpaces.Len() };
-}
-
 void FYamlTokenizer::AdvanceReadLines()
 {
 	ReadLine = NextLine;
@@ -997,19 +1033,6 @@ void FYamlTokenizer::AdvanceReadLines()
 	}
 }
 
-void FYamlTokenizer::SetLastError(FUtf8StringView Message)
-{
-	LastError.Reset();
-	if (Message.Len())
-	{
-		LastError.Appendf("'%.*s' in line %u: '%.*s'",
-			Message.Len(), Message.GetData(),
-			ReadLine.LineNumber,
-			ReadLine.FullView.Len(),
-			ReadLine.FullView.GetData());
-	}
-}
-
 ////////////////////////////////////////////////////////////////////////////////
 
 bool GrabToNextBeginStruct(FToken& Out, FScopedTokenReader& ScopedReader)
