nss: a JSON markup tokeniser

#rnx

#Synced-CL 50479750

diff --git a/Engine/Plugins/Experimental/PlainProps/Source/Private/PlainPropsParse.cpp b/Engine/Plugins/Experimental/PlainProps/Source/Private/PlainPropsParse.cpp
--- a/Engine/Plugins/Experimental/PlainProps/Source/Private/PlainPropsParse.cpp
+++ b/Engine/Plugins/Experimental/PlainProps/Source/Private/PlainPropsParse.cpp
@@ -1033,6 +1033,355 @@ void FYamlTokenizer::AdvanceReadLines()
 	}
 }
 
+
+
+////////////////////////////////////////////////////////////////////////////////
+
+class FJsonTokenizer
+	: public FMarkupTokenizer
+{
+public:
+								FJsonTokenizer(FUtf8StringView InData);
+								~FJsonTokenizer();
+	virtual bool				HasMore() const override;
+	virtual TOptional<FToken>	GrabToken() override;
+
+private:
+	enum class EScope : uint8 { Document, Struct, Range, };
+
+	enum : uint16
+	{
+		Flag_AllowStruct	= 1 << 0,
+		Flag_AllowRange		= 1 << 1,
+		Flag_AllowOpen		= Flag_AllowStruct | Flag_AllowRange,
+		Flag_AllowClose		= 1 << 2,
+		Flag_ExpectKey		= 1 << 3,
+		Flag_ExpectColon	= 1 << 4,
+		Flag_ExpectValue	= 1 << 5 | Flag_AllowOpen,
+		Flag_ExpectComma	= 1 << 6,
+		Flag_HasComma		= 1 << 7,
+	};
+
+	struct FScopeInfo
+	{
+		EScope			Scope;
+		uint16			Counter;
+		uint8			Flags;
+	};
+
+	TOptional<FToken>	GrabTokenInternal();
+	FToken				SetLastError(FUtf8StringView Msg);
+	void				AdvanceData(uint32 ChopLength);
+	TOptional<FToken>	ParseOpen();
+	TOptional<FToken>	ParseClose();
+	TOptional<FToken>	ParseKey();
+	TOptional<FToken>	ParseColon();
+	TOptional<FToken>	ParseValue();
+	TOptional<FToken>	ParseComma();
+	FUtf8StringView		Data;
+	TArray<FScopeInfo>	Scopes;
+	TOptional<FToken>	PendingToken;
+
+};
+
+FJsonTokenizer::FJsonTokenizer(FUtf8StringView InData)
+: FMarkupTokenizer(InData)
+{
+	Scopes.Emplace(FScopeInfo{
+		.Flags = Flag_AllowStruct,
+	});
+
+	AdvanceData(0);
+}
+
+FJsonTokenizer::~FJsonTokenizer()
+{
+	check(!PendingToken.IsSet());
+	check(Scopes.IsEmpty());
+}
+
+bool FJsonTokenizer::HasMore() const
+{
+	return !Scopes.IsEmpty();
+}
+
+void FJsonTokenizer::AdvanceData(uint32 ChopLength)
+{
+	Data.RightChopInline(ChopLength);
+	FUtf8StringView Trim = FAsciiSet::FindPrefixWith(Data, WhitespaceSet);
+	Data.RightChopInline(Trim.Len());
+
+	if (!Data.IsEmpty())
+	{
+		return;
+	}
+
+	FLineInfo Line = ReadOneLine();
+	if (Line.EndOfFile)
+	{
+		return;
+	}
+
+	Data = Line.View;
+}
+
+FToken FJsonTokenizer::SetLastError(FUtf8StringView Msg)
+{
+	FMarkupTokenizer::SetLastError(Msg, Data);
+	Scopes.Reset();
+	return FToken{ EToken::Invalid, 0 };
+}
+
+TOptional<FToken> FJsonTokenizer::ParseOpen()
+{
+	FScopeInfo& Last = Scopes.Last();
+
+	int32 Char = Data[0];
+	if (Char == '[')
+	{
+		AdvanceData(1);
+
+		if ((Last.Flags & Flag_AllowRange) == 0)
+		{
+			return SetLastError("Range not allowed here");
+		}
+
+		Last.Flags = Flag_ExpectComma | Flag_AllowClose;
+
+		Scopes.Emplace(FScopeInfo{
+			.Scope = EScope::Range,
+			.Counter = 1,
+			.Flags = Flag_ExpectValue | Flag_AllowClose,
+		});
+		uint16 Depth = uint16(Scopes.Num() - 1);
+
+		return FToken{ EToken::BeginRange, Depth };
+	}
+
+	if (Char == '{')
+	{
+		check(Last.Flags & Flag_AllowStruct);
+
+		AdvanceData(1);
+
+		if (Last.Flags == Flag_AllowStruct)
+		{
+			check(Scopes.Num() == 1);
+			Last.Flags = Flag_ExpectKey | Flag_AllowClose;
+			return FToken{ EToken::BeginDocument, 0 };
+		}
+
+		Last.Flags = Flag_ExpectComma | Flag_AllowClose;
+
+		Scopes.Emplace(FScopeInfo{
+			.Scope = EScope::Struct,
+			.Counter = 0,
+			.Flags = Flag_ExpectKey | Flag_AllowClose,
+		});
+		uint16 Depth = uint16(Scopes.Num() - 1);
+
+		return FToken{ EToken::BeginStruct, Depth };
+	}
+
+	return {};
+}
+
+TOptional<FToken> FJsonTokenizer::ParseClose()
+{
+	int32 Char = Data[0];
+	if ((Char != ']') && (Char != '}'))
+	{
+		return {};
+	}
+
+	const FScopeInfo& Last = Scopes.Last();
+
+	if (Last.Flags & Flag_HasComma)
+	{
+		return SetLastError("Trailing seperator");
+	}
+
+	bool bValid = (Char == '}') & (Last.Scope != EScope::Range);
+	bValid |= (Char == ']') & (Last.Scope == EScope::Range);
+	if (!bValid)
+	{
+		return SetLastError("Mismatched scope");
+	}
+
+	AdvanceData(1);
+
+	Scopes.Pop(EAllowShrinking::No);
+
+	EToken Token = EToken::EndRange;
+	if (Char == '}')
+	{
+		Token = (Scopes.IsEmpty()) ? EToken::EndDocument : EToken::EndStruct;
+	}
+
+	return FToken{ Token, uint16(Scopes.Num()) };
+}
+
+TOptional<FToken> FJsonTokenizer::ParseKey()
+{
+	int32 Char = Data[0];
+	if (Char != '\"')
+	{
+		return SetLastError("Key string expected");
+	}
+
+	FLineInfo LineInfo = { .View = Data, };
+	FUtf8StringView Key = GrabTokenFromLine(LineInfo);
+	if (Key.IsEmpty())
+	{
+		return SetLastError("Empty key");
+	}
+
+	AdvanceData(Key.Len() + 2);
+
+	Scopes.Last().Flags = Flag_ExpectColon;
+	return FToken{ EToken::LeafId, uint16(Scopes.Num()), Key };
+}
+
+TOptional<FToken> FJsonTokenizer::ParseColon()
+{
+	int32 Char = Data[0];
+	if (Char != ':')
+	{
+		return SetLastError("Key should be followed by a colon");
+	}
+
+	AdvanceData(1);
+
+	Scopes.Last().Flags = Flag_ExpectValue;
+	return {};
+}
+
+TOptional<FToken> FJsonTokenizer::ParseValue()
+{
+	int32 Char = Data[0];
+
+	if (TOptional<FToken> Token = ParseOpen(); Token.IsSet())
+	{
+		return Token;
+	}
+
+	if (Char != '\"')
+	{
+		return SetLastError("Only string-type values are supported");
+	}
+
+	FLineInfo LineInfo = { .View = Data, };
+	FUtf8StringView Value = GrabTokenFromLine(LineInfo);
+	AdvanceData(Value.Len() + 2);
+
+	Scopes.Last().Flags = Flag_ExpectComma | Flag_AllowClose;
+
+	EToken Token = (Scopes.Last().Scope == EScope::Range) ? EToken::Leaf : EToken::LeafValue;
+	return FToken{ Token, uint16(Scopes.Num()), Value };
+}
+
+TOptional<FToken> FJsonTokenizer::ParseComma()
+{
+	int32 Char = Data[0];
+	if (Char != ',')
+	{
+		return SetLastError("Comma separator expected");
+	}
+
+	AdvanceData(1);
+
+	FScopeInfo& Last = Scopes.Last();
+	Last.Flags = (Last.Scope == EScope::Range) ? Flag_ExpectValue : Flag_ExpectKey;
+	Last.Flags |= Flag_HasComma;
+	return {};
+}
+
+TOptional<FToken> FJsonTokenizer::GrabTokenInternal()
+{
+	if (Data.IsEmpty())
+	{
+		return SetLastError("Encountered EOF");
+	}
+
+	if (Scopes.IsEmpty())
+	{
+		return SetLastError("Unexpected token");
+	}
+
+	static constexpr struct {
+		uint32					Flag;
+		TOptional<FToken>		(FJsonTokenizer::*Method)(void);
+	} Rota[] = {
+		{ Flag_AllowOpen,		&FJsonTokenizer::ParseOpen },
+		{ Flag_AllowClose,		&FJsonTokenizer::ParseClose },
+		{ Flag_ExpectComma, 	&FJsonTokenizer::ParseComma },
+		{ Flag_ExpectKey,		&FJsonTokenizer::ParseKey },
+		{ Flag_ExpectColon,		&FJsonTokenizer::ParseColon },
+		{ Flag_ExpectValue,		&FJsonTokenizer::ParseValue },
+	};
+
+	for (const FScopeInfo& Last = Scopes.Last();;)
+	{
+		for (const auto& [RotaFlag, Method] : Rota)
+		{
+			if (Last.Flags & RotaFlag)
+			{
+				if (TOptional<FToken> Token = (this->*Method)(); Token.IsSet())
+				{
+					return Token;
+				}
+			}
+		}
+
+		if (Last.Flags & (Flag_ExpectKey | Flag_ExpectValue))
+		{
+			continue;
+		}
+
+		SetLastError("Parsing reached the unreachable");
+		return {};
+	}
+}
+
+TOptional<FToken> FJsonTokenizer::GrabToken()
+{
+	TOptional<FToken> Token;
+	if (PendingToken.IsSet())
+	{
+		Token = PendingToken.GetValue();
+		PendingToken.Reset();
+	}
+	else
+	{
+		Token = GrabTokenInternal();
+	}
+
+	if (!Token.IsSet())						return {};
+	if (Token->Token == EToken::Invalid)	return {};
+	if (Token->Token != EToken::LeafId)		return Token;
+
+	TOptional<FToken> NextToken = GrabTokenInternal();
+	if (NextToken.IsSet())
+	{
+		if (NextToken->Token == EToken::Invalid)
+		{
+			return {};
+		}
+
+		if (NextToken->Token == EToken::BeginStruct || NextToken->Token == EToken::BeginRange)
+		{
+			NextToken->Str = Token->Str;
+			NextToken->Len = Token->Len;
+			return NextToken;
+		}
+	}
+
+	PendingToken = NextToken;
+	return Token;
+}
+
+
+
 ////////////////////////////////////////////////////////////////////////////////
 
 bool GrabToNextBeginStruct(FToken& Out, FScopedTokenReader& ScopedReader)
