nss: Removed unused member.

#rnx

#Synced-CL 50479750

diff --git a/Engine/Plugins/Experimental/PlainProps/Source/Private/PlainPropsParse.cpp b/Engine/Plugins/Experimental/PlainProps/Source/Private/PlainPropsParse.cpp
--- a/Engine/Plugins/Experimental/PlainProps/Source/Private/PlainPropsParse.cpp
+++ b/Engine/Plugins/Experimental/PlainProps/Source/Private/PlainPropsParse.cpp
@@ -1065,7 +1065,6 @@ private:
 	struct FScopeInfo
 	{
 		EScope			Scope;
-		uint16			Counter;
 		uint8			Flags;
 	};
 
@@ -1150,7 +1149,6 @@ TOptional<FToken> FJsonTokenizer::ParseOpen()
 
 		Scopes.Emplace(FScopeInfo{
 			.Scope = EScope::Range,
-			.Counter = 1,
 			.Flags = Flag_ExpectValue | Flag_AllowClose,
 		});
 		uint16 Depth = uint16(Scopes.Num() - 1);
@@ -1175,7 +1173,6 @@ TOptional<FToken> FJsonTokenizer::ParseOpen()
 
 		Scopes.Emplace(FScopeInfo{
 			.Scope = EScope::Struct,
-			.Counter = 0,
 			.Flags = Flag_ExpectKey | Flag_AllowClose,
 		});
 		uint16 Depth = uint16(Scopes.Num() - 1);
@@ -2572,4 +2569,19 @@ void ParseYamlBatch(TArray64<uint8>& OutBinary, FUtf8StringView Yaml)
 	BatchParser.Parse(OutBinary);
 }
 
+
+static struct X
+{
+	X()
+	{
+		FJsonTokenizer t("\n\n {\n\n\n  \n\"sdf\"\n  \n\n:\n\n\n \n\n\n\"x\"\n\n}\n\n \n\n");
+		while (t.HasMore())
+		{
+			auto tt = t.GrabToken();
+			tt->Token = EToken::Invalid;
+		}
+	}
+} xyz;
+
+
 } // namespace PlainProps
