nss: Fixed early out on run of empty lines

#rnx

#Synced-CL 50479750

diff --git a/Engine/Plugins/Experimental/PlainProps/Source/Private/PlainPropsParse.cpp b/Engine/Plugins/Experimental/PlainProps/Source/Private/PlainPropsParse.cpp
--- a/Engine/Plugins/Experimental/PlainProps/Source/Private/PlainPropsParse.cpp
+++ b/Engine/Plugins/Experimental/PlainProps/Source/Private/PlainPropsParse.cpp
@@ -1115,11 +1115,16 @@ void FJsonTokenizer::AdvanceData(uint32 ChopLength)
 		return;
 	}
 
-	FLineInfo Line = ReadOneLine();
-	if (Line.EndOfFile)
+	FLineInfo Line;
+	do
 	{
-		return;
+		Line = ReadOneLine();
+		if (Line.EndOfFile)
+		{
+			return;
+		}
 	}
+	while (Line.View.IsEmpty());
 
 	Data = Line.View;
 }
@@ -2569,19 +2574,4 @@ void ParseYamlBatch(TArray64<uint8>& OutBinary, FUtf8StringView Yaml)
 	BatchParser.Parse(OutBinary);
 }
 
-
-static struct X
-{
-	X()
-	{
-		FJsonTokenizer t("\n\n {\n\n\n  \n\"sdf\"\n  \n\n:\n\n\n \n\n\n\"x\"\n\n}\n\n \n\n");
-		while (t.HasMore())
-		{
-			auto tt = t.GrabToken();
-			tt->Token = EToken::Invalid;
-		}
-	}
-} xyz;
-
-
 } // namespace PlainProps
