nss: Add JSON equivalents for YAML print/parse functions, and include JSON rountrips as part of NSS program tests

#rnx

#Synced-CL 50479750

diff --git a/Engine/Plugins/Experimental/PlainProps/Source/Private/PlainPropsInternalPrint.h b/Engine/Plugins/Experimental/PlainProps/Source/Private/PlainPropsInternalPrint.h
--- a/Engine/Plugins/Experimental/PlainProps/Source/Private/PlainPropsInternalPrint.h
+++ b/Engine/Plugins/Experimental/PlainProps/Source/Private/PlainPropsInternalPrint.h
@@ -31,6 +31,7 @@ using FYamlBuilderPtr = TUniquePtr<FMarkupBuilder, FMarkupBuilderDeleter>;
 using FJsonBuilderPtr = TUniquePtr<FMarkupBuilder, FMarkupBuilderDeleter>;
 
 [[nodiscard]] FYamlBuilderPtr MakeYamlBuilder(FUtf8StringBuilderBase& StringBuilder);
+[[nodiscard]] FJsonBuilderPtr MakeJsonBuilder(FUtf8StringBuilderBase& StringBuilder);
 
 ///////////////////////////////////////////////////////////////////////////////
 
diff --git a/Engine/Plugins/Experimental/PlainProps/Source/Private/PlainPropsInternalTest.h b/Engine/Plugins/Experimental/PlainProps/Source/Private/PlainPropsInternalTest.h
--- a/Engine/Plugins/Experimental/PlainProps/Source/Private/PlainPropsInternalTest.h
+++ b/Engine/Plugins/Experimental/PlainProps/Source/Private/PlainPropsInternalTest.h
@@ -12,6 +12,7 @@ namespace PlainProps
 class FSchemaBatchId;
 struct FStructView;
 
-[[nodiscard]] FSchemaBatchId ParseBatch(TArray64<uint8>& OutData, TArray<FStructView>& OutObjects, FUtf8StringView YamlView);
+[[nodiscard]] FSchemaBatchId ParseYamlBatch(TArray64<uint8>& OutData, TArray<FStructView>& OutObjects, FUtf8StringView YamlView);
+[[nodiscard]] FSchemaBatchId ParseJsonBatch(TArray64<uint8>& OutData, TArray<FStructView>& OutObjects, FUtf8StringView JsonView);
 
 } // namespace PlainProps
diff --git a/Engine/Plugins/Experimental/PlainProps/Source/Private/PlainPropsParse.cpp b/Engine/Plugins/Experimental/PlainProps/Source/Private/PlainPropsParse.cpp
--- a/Engine/Plugins/Experimental/PlainProps/Source/Private/PlainPropsParse.cpp
+++ b/Engine/Plugins/Experimental/PlainProps/Source/Private/PlainPropsParse.cpp
@@ -2566,6 +2566,14 @@ void FBatchParser::Write(TArray64<uint8>& Out)
 
 ///////////////////////////////////////////////////////////////////////////////
 
+void ParseJsonBatch(TArray64<uint8>& OutBinary, FUtf8StringView Json)
+{
+	FScratchAllocator Scratch;
+	FJsonTokenizer JsonScanner(Json);
+	FBatchParser BatchParser(JsonScanner, Scratch);
+	BatchParser.Parse(OutBinary);
+}
+
 void ParseYamlBatch(TArray64<uint8>& OutBinary, FUtf8StringView Yaml)
 {
 	FScratchAllocator Scratch;
diff --git a/Engine/Plugins/Experimental/PlainProps/Source/Private/PlainPropsPrint.cpp b/Engine/Plugins/Experimental/PlainProps/Source/Private/PlainPropsPrint.cpp
--- a/Engine/Plugins/Experimental/PlainProps/Source/Private/PlainPropsPrint.cpp
+++ b/Engine/Plugins/Experimental/PlainProps/Source/Private/PlainPropsPrint.cpp
@@ -812,6 +812,13 @@ void PrintYamlBatch(FUtf8Builder& Out, const FBatchIds& Ids, TConstArrayView<FSt
 	PrintBatch(YamlBuilder, Printer, Objects);
 }
 
+void PrintJsonBatch(FUtf8Builder& Out, const FBatchIds& Ids, TConstArrayView<FStructView> Objects)
+{
+	FJsonBuilder JsonBuilder(Out);
+	FBatchPrinter Printer(JsonBuilder, Ids);
+	PrintBatch(JsonBuilder, Printer, Objects);
+}
+
 ///////////////////////////////////////////////////////////////////////////////
 
 FStructSchemaReader::FStructSchemaReader(const FStructSchema& Schema, FSchemaBatchId InBatch)
diff --git a/Engine/Plugins/Experimental/PlainProps/Source/Public/PlainPropsParse.h b/Engine/Plugins/Experimental/PlainProps/Source/Public/PlainPropsParse.h
--- a/Engine/Plugins/Experimental/PlainProps/Source/Public/PlainPropsParse.h
+++ b/Engine/Plugins/Experimental/PlainProps/Source/Public/PlainPropsParse.h
@@ -46,5 +46,6 @@ template<> PLAINPROPS_API TOptional<uint64>				ParseHex(FUtf8StringView String);
 ///////////////////////////////////////////////////////////////////////////////
 
 void PLAINPROPS_API ParseYamlBatch(TArray64<uint8>& OutBinary, FUtf8StringView Yaml);
+void PLAINPROPS_API ParseJsonBatch(TArray64<uint8>& OutBinary, FUtf8StringView Json);
 
 } // namespace PlainProps
diff --git a/Engine/Plugins/Experimental/PlainProps/Source/Public/PlainPropsPrint.h b/Engine/Plugins/Experimental/PlainProps/Source/Public/PlainPropsPrint.h
--- a/Engine/Plugins/Experimental/PlainProps/Source/Public/PlainPropsPrint.h
+++ b/Engine/Plugins/Experimental/PlainProps/Source/Public/PlainPropsPrint.h
@@ -58,6 +58,7 @@ template<> PLAINPROPS_API void PrintHex(FUtf8Builder& Out, uint64 Value);
 
 PLAINPROPS_API void PrintMemberSchema(FUtf8Builder& Out, const FIds& Ids, FMemberSchema Member);
 PLAINPROPS_API void PrintYamlBatch(FUtf8Builder& Out, const FBatchIds& Ids, TConstArrayView<FStructView> Objects);
+PLAINPROPS_API void PrintJsonBatch(FUtf8Builder& Out, const FBatchIds& Ids, TConstArrayView<FStructView> Objects);
 PLAINPROPS_API void PrintDiff(FUtf8Builder& Out, const FIds& Ids, const FDiffPath& Diff);
 PLAINPROPS_API void PrintDiff(FUtf8Builder& Out, const FBatchIds& Ids, const FReadDiffPath& Diff);
 PLAINPROPS_API void PrintDiffs(FUtf8Builder& Out, const FIds& Ids, TConstArrayView<FDiffPath> Diffs);
diff --git a/Engine/Plugins/Experimental/PlainProps/Source/Tests/PlainPropsTests.cpp b/Engine/Plugins/Experimental/PlainProps/Source/Tests/PlainPropsTests.cpp
--- a/Engine/Plugins/Experimental/PlainProps/Source/Tests/PlainPropsTests.cpp
+++ b/Engine/Plugins/Experimental/PlainProps/Source/Tests/PlainPropsTests.cpp
@@ -322,11 +322,8 @@ public:
 	virtual void AppendString(FUtf8StringBuilderBase& Out, FNameId Name) const override { Out.Append(Names[Name]); }
 };
 
-[[nodiscard]] FSchemaBatchId ParseBatch(TArray64<uint8>& OutData, TArray<FStructView>& OutObjects, FUtf8StringView YamlView)
+static FSchemaBatchId ParseBatch(TArray64<uint8>& OutData, TArray<FStructView>& OutObjects)
 {
-	// Parse yaml
-	ParseYamlBatch(OutData, YamlView);
-
 	// Grab and mount parsed schemas
 	FByteReader It(MakeMemoryView(OutData));
 	const uint32 SchemasSize = It.Grab<uint32>();
@@ -343,23 +340,49 @@ public:
 	}
 	
 	return Batch;
+};
+
+[[nodiscard]] FSchemaBatchId ParseYamlBatch(TArray64<uint8>& OutData, TArray<FStructView>& OutObjects, FUtf8StringView YamlView)
+{
+	ParseYamlBatch(OutData, YamlView);
+	return ParseBatch(OutData, OutObjects);
+};
+
+[[nodiscard]] FSchemaBatchId ParseJsonBatch(TArray64<uint8>& OutData, TArray<FStructView>& OutObjects, FUtf8StringView JsonView)
+{
+	ParseJsonBatch(OutData, JsonView);
+	return ParseBatch(OutData, OutObjects);
 }
 
-static void RoundtripText(const FBatchIds& BatchIds, TConstArrayView<FStructView> Objects)
+enum class ERoundtripFormat { Yaml, Json };
+
+static void RoundtripText(
+	const FBatchIds& BatchIds,
+	TConstArrayView<FStructView> Objects,
+	ERoundtripFormat Format=ERoundtripFormat::Yaml)
 {
-	// Print yaml
-	TUtf8StringBuilder<4096> Yaml;
-	PrintYamlBatch(Yaml, BatchIds, Objects);
-	FUtf8StringView YamlView = Yaml.ToView();
+	// Print text
+	TUtf8StringBuilder<4096> Text;
+	switch (Format)
+	{
+	case ERoundtripFormat::Yaml: PrintYamlBatch(Text, BatchIds, Objects); break;
+	case ERoundtripFormat::Json: PrintJsonBatch(Text, BatchIds, Objects); break;
+	}
+	FUtf8StringView TextView = Text.ToView();
 
-	// Log yaml
-	auto Wide = StringCast<TCHAR>(YamlView.GetData(), YamlView.Len());
+	// Log text
+	auto Wide = StringCast<TCHAR>(TextView.GetData(), TextView.Len());
 	UE_LOG(LogPlainPropsTests, Log, TEXT("Schemas with StableNames:\n%.*s"), Wide.Length(), Wide.Get());
 
-	// Parse yaml
+	// Parse text
 	TArray64<uint8> Data;
 	TArray<FStructView> ParsedObjects;
-	FSchemaBatchId ParsedBatch = ParseBatch(Data, ParsedObjects, YamlView);
+	FSchemaBatchId ParsedBatch;
+	switch (Format)
+	{
+	case ERoundtripFormat::Yaml: ParsedBatch = ParseYamlBatch(Data, ParsedObjects, TextView); break;
+	case ERoundtripFormat::Json: ParsedBatch = ParseJsonBatch(Data, ParsedObjects, TextView); break;
+	}
 
 	// Diff schemas
 	CHECK_FALSE(DiffSchemas(BatchIds.GetBatchId(), ParsedBatch));
@@ -419,9 +442,9 @@ public:
 		BatchIds.Emplace(Names, Batch);
 	}
 
-	void RoundtripText()
+	void RoundtripText(PlainProps::ERoundtripFormat Format)
 	{
-		PlainProps::RoundtripText(BatchIds.GetValue(), Objects);
+		PlainProps::RoundtripText(BatchIds.GetValue(), Objects, Format);
 	}
 
 	~FTestBatchReader()
@@ -450,9 +473,17 @@ static void TestSerialize(void (*BuildObjects)(FTestBatchBuilder&, FScratchAlloc
 		Data = Batch.Write();
 	}
 
-	FTestBatchReader Batch(MakeMemoryView(Data));
-	Batch.RoundtripText();
-	CheckObjects(Batch.GetObjects(), Batch.GetNames());
+	{
+		FTestBatchReader Batch(MakeMemoryView(Data));
+		Batch.RoundtripText(PlainProps::ERoundtripFormat::Yaml);
+		CheckObjects(Batch.GetObjects(), Batch.GetNames());
+	}
+
+	{
+		FTestBatchReader Batch(MakeMemoryView(Data));
+		Batch.RoundtripText(PlainProps::ERoundtripFormat::Json);
+		CheckObjects(Batch.GetObjects(), Batch.GetNames());
+	}
 }
 
 ////////////////////////////////////////////////////////////////////////////////////////////////
diff --git a/Engine/Plugins/Experimental/PlainProps/Source/Tests/PlainPropsUeCoreTests.cpp b/Engine/Plugins/Experimental/PlainProps/Source/Tests/PlainPropsUeCoreTests.cpp
--- a/Engine/Plugins/Experimental/PlainProps/Source/Tests/PlainPropsUeCoreTests.cpp
+++ b/Engine/Plugins/Experimental/PlainProps/Source/Tests/PlainPropsUeCoreTests.cpp
@@ -295,7 +295,7 @@ static void RoundtripText(const FBatchIds& BatchIds, TConstArrayView<FStructView
 	// Parse yaml
 	TArray64<uint8> Data;
 	TArray<FStructView> ParsedObjects;
-	FSchemaBatchId ParsedBatch = ParseBatch(Data, ParsedObjects, YamlView);
+	FSchemaBatchId ParsedBatch = ParseYamlBatch(Data, ParsedObjects, YamlView);
 
 	if (Format == ESchemaFormat::StableNames)
 	{
