PlainProps upgrade system improvements

* add range support and unit tests
* improved opaque and named enumerator/flag support
* improved dynamic struct support
* fixed various bugs
* always loosen opaque input enums to prevent mixed loose/source enum ranges
* implement missing loosening paths
* implement missing decomposition paths
* implement missing recomposition paths
* implement missing fastening paths
* add FFriend mechanism
* add debug print capabilities
* improved error messages

#rnx
#Synced-CL 50709972

diff --git a/Engine/Plugins/Experimental/PlainProps/Source/Private/PlainPropsInternalUpgrade.h b/Engine/Plugins/Experimental/PlainProps/Source/Private/PlainPropsInternalUpgrade.h
--- a/Engine/Plugins/Experimental/PlainProps/Source/Private/PlainPropsInternalUpgrade.h
+++ b/Engine/Plugins/Experimental/PlainProps/Source/Private/PlainPropsInternalUpgrade.h
@@ -4,9 +4,12 @@
 
 #include "PlainPropsUpgrade.h"
 #include "PlainPropsBuild.h"
+#include "PlainPropsInternalBind.h"
+#include "PlainPropsInternalRead.h"
 #include "PlainPropsRead.h"
 #include "Containers/Map.h"
 #include "Containers/Set.h"
+#include "Containers/StaticArray.h"
 
 namespace PlainProps::Upgrade
 {
@@ -45,7 +48,7 @@ struct FParameterIndexer : private FParameterIdSet
 	uint32				Num() const								{ return uint32(TSet<FOptionalMemberId>::Num()); }
 };
 
-class FLooseIndexers 
+class FLooseIndexers
 {
 public:
 	FLooseIndexers(const IDeclarations& Types) : Declarations(Types) {}
@@ -77,7 +80,7 @@ struct FLooseParameter
 struct FOutParameter
 {
 	FParameterId						Name;
-	uint8								NumInners = 0;	// Struct members or enum enumerators
+	uint8								NumInners = 0;	// Named struct members or enumerators/flags
 	uint8								RecompIdx = 0;	// Only for enums with NumInners > 0
 	FLooseType							Type;
 };
@@ -174,7 +177,7 @@ private:
 	FEnumLooseners					Looseners;
 
 	FRangeMember					ToLooseRange(FParameterId Name, FLooseType Type, FRangeView Range);
-	const FLooseNestedRange*		LoosenRanges(FNestedRangeView Range);
+	FLooseNestedRange*				LoosenRanges(FNestedRangeView Range);
 };
 
 ////////////////////////////////////////////////////////////////////////////////////////////////
@@ -221,7 +224,7 @@ union FMatchIdx
 
 struct FMemberMatch
 {
-	FLooseType	Type;
+	FLooseType	Type; // With Innermost.IsDecomposed == 0
 	FMatchIdx	Idx;
 	uint8		InputIdx;
 	uint8		Recipient;	// 0 = this, 1 = Outer, 2 = Outer->Outer, ...
@@ -273,13 +276,14 @@ struct FVersionUpgrades
 // Source constant -> loose value
 struct FFlatEnumLoosener
 {
+	// todo: either track and destroy loosener or change to hashtable w/o internal allocation (e.g. null/footer container allocator, some fixed hash map) or sorted array lookup
 	TMap<uint64, FEnumeratorId> ConstantNames;
 };
 
 // Source flag -> loose value
 struct FFlagEnumLoosener
 {
-	TArray<FEnumeratorId, TFixedAllocator<64>> FlagNames;
+	TStaticArray<FEnumeratorId, 64> FlagNames{FEnumeratorId{}};
 };
 
 // Helps convert source enums to loose or indicate that declaration matches source schema
@@ -320,11 +324,21 @@ struct FEnumTransformation
 // Cached info for upgrading a struct
 struct FTransformation
 {
-	const FLooseSchema*			SourceSchema = nullptr;
+	mutable const FLooseSchema*	SourceSchema = nullptr;
 	const FStructUpgrades*		Upgrades = nullptr;
 	const FLooseSchema*			TargetSchema = nullptr;
 };
 
+// Cached derived schema metadata from a source batch
+struct FSourceSchemas
+{
+	uint32								NumStructs = 0;
+	const FSchemaBatch*					Batch = nullptr;
+	FInnerIds							RuntimeIds;
+	TArray<uint16>						StructVersions;
+	TBitArray<>							EnumFlagModes;
+};
+
 // Cached info for upgrading all structs and enums in a source batch
 struct FTransformations
 {
@@ -332,11 +346,11 @@ struct FTransformations
 	TSet<FDeclId>				StructLookup;		// If we start removing, switch to TSparseSet for FSetElementId stability
 	TArray<FEnumTransformation>	Enums;
 	TSet<FEnumId>				EnumLookup;			// If we start removing, switch to TSparseSet for FSetElementId stability
-	int32						NumSourceStructs;
+	FSourceSchemas				Schemas;
 
 	const FTransformation&		operator[](FStructSchemaId Id) const	{ return Structs[Id.Idx]; }
 	const FTransformation&		operator[](FDeclId Id) const			{ return Structs[StructLookup.FindId(Id).AsInteger()]; }	// TArrayView::RangeCheck checks invalid FSetElementId
-	const FEnumTransformation&	operator[](FEnumSchemaId Id) const		{ return Enums[Id.Idx - NumSourceStructs]; }
+	const FEnumTransformation&	operator[](FEnumSchemaId Id) const		{ return Enums[Id.Idx - Schemas.NumStructs]; }
 	const FEnumTransformation&	operator[](FEnumId Id) const			{ return Enums[EnumLookup.FindId(Id).AsInteger()]; }		// TArrayView::RangeCheck checks invalid FSetElementId
 };
 
@@ -386,6 +400,19 @@ struct FOutputScope
 	const FOuts					Outputs;
 };
 
+// Helps print some loose types on top of the FDebugIds API
+class FDebugPrinter : public FDebugIds
+{
+	const FLooseIndexers&		LooseIds;
+public:
+	FDebugPrinter(const FIds& Ids, FLooseIndexers& Loose) : FDebugIds(Ids), LooseIds(Loose) {}
+
+	using FDebugIds::Print;
+	FString						Print(FLooseType Type) const;
+	FString						Print(FDeclId Id, FParameterId Name) const	{ return Print(LooseIds.Get(Id).Resolve(Name)); }
+	FString						Print(FEnumId Id, FEnumeratorId Name) const	{ return Print(LooseIds.Get(Id).Resolve(Name)); }
+};
+
 // Context for upgrading a source batch
 struct FBatchContext
 {
@@ -398,7 +425,7 @@ struct FBatchContext
 	const FSchemaBatchId		Batch; 
 	FOutputStack				OutputStack;
 	FContext					OpCtx;
-	FDebugIds					Debug;
+	FDebugPrinter				Debug;
 };
 
 // Home for FIn of "polyary" (>1 inputs) ops
@@ -439,6 +466,7 @@ class FInstanceUpgrader
 	void							Give(FLooseType Type, FLooseValue Value);
 	void							Give(FLooseType Type, FLooseValue Value, FMatchHelper Helper);
 	FLooseValue						Decompose(FLooseType Type, FLooseValue Value, uint16 DecompIdx);
+	void							Decompose(FLooseType Type, FLooseNestedRange Range, const void* Decomp);
 	FLooseValue						Decompose(FLooseType Type, FLooseValue Value, const FEnumDecomposition& Decomp);
 	FLooseValue						Decompose(FEnumMember Input, const FEnumDecomposition& Decomp);
 	FLooseValue						Decompose(FStructMember Input, const FStructDecomposition& Decomp);
@@ -450,11 +478,13 @@ class FInstanceUpgrader
 	void							Run(FOp Op, const FIn* Inputs, uint8 Recipient);
 	FOuter							Resolve(uint8 Recipient);
 	void							Recompose(TArrayView<FOut> Outputs, TConstArrayView<FOutParameter> Expected);
-	FEnumeratorData					Recompose(FInnermostEnumType Type, FEnumIndexData In, TConstArrayView<FEnumeratorId> Order);
+	FEnumeratorData					Recompose(FEnumIndexData In, TConstArrayView<FEnumeratorId> Order, FInnermostEnumType Type);
+	void							Recompose(FRangeData Out, uint64 Num, uint8 Depth, TConstArrayView<FEnumeratorId> Order, FInnermostEnumType Type);
 	void							Flush();
 	FLooseStruct					Finish(EMismatch& OutMismatch);
 
 	FString							Print(FParameterId Parameter) const;
+	FString							Print(FLooseType Type) const;
 	template<class T> FString		Print(T Item) const { return Ctx.Debug.Print(Item); }
 	
 public:
@@ -480,7 +510,13 @@ struct FSchemaCache
 template <typename T>
 struct TDestroyOnly
 {
-	void operator()(T* Ptr) const { Ptr->~T(); } 
+	void operator()(T* Ptr) const
+	{
+		if (Ptr)
+		{
+			Ptr->~T();
+		}
+	} 
 };
 
 template<class T>
diff --git a/Engine/Plugins/Experimental/PlainProps/Source/Private/PlainPropsUpgrade.cpp b/Engine/Plugins/Experimental/PlainProps/Source/Private/PlainPropsUpgrade.cpp
--- a/Engine/Plugins/Experimental/PlainProps/Source/Private/PlainPropsUpgrade.cpp
+++ b/Engine/Plugins/Experimental/PlainProps/Source/Private/PlainPropsUpgrade.cpp
@@ -18,6 +18,92 @@ namespace PlainProps
 static_assert(sizeof(FLooseType) == sizeof(uint64), "AsInt() requirement");
 static_assert(std::has_unique_object_representations_v<FLooseType>, "AsInt() requirement");
 
+inline bool IsDecomposed(FInnermostType T)
+{
+	return (T.Kind == EInnermostKind::Struct && T.Struct.IsDecomposed) || (T.Kind==EInnermostKind::Enum && T.Enum.IsDecomposed);
+}
+
+static FLooseType ClearDecomposed(FLooseType Type)
+{
+	if (Type.InnermostIsStruct())
+	{
+		check(Type.InnermostType.Struct.IsDecomposed);
+		Type.InnermostType.Struct.IsDecomposed = 0;
+	}
+	else
+	{
+		check(Type.InnermostIsEnum());
+		check(Type.InnermostType.Enum.IsDecomposed);
+		Type.InnermostType.Enum.IsDecomposed = 0;
+	}
+	return Type;
+}
+
+inline bool CompareDecomposedAndRecomposed(FLooseType Decomposed, FLooseType Recomposed)
+{
+	return ClearDecomposed(Decomposed).CompareAll(Recomposed);
+}
+
+void FLooseMember::RecordDynamicType(FLooseType In)
+{
+	check(Type.InnermostType == InnermostDynamicStruct);
+	check(!Type.InnermostId || Type.InnermostId == In.InnermostId && Type.InnermostVersion == In.InnermostVersion);
+	Type.InnermostId = In.InnermostId;
+	Type.InnermostVersion = In.InnermostVersion;
+}
+
+static void DebugPrintVersion(FString& Out, uint16 Version)
+{
+	if (Version)
+	{
+		Out.AppendChar('@');
+		Out.AppendInt(Version);
+	}
+}
+
+static FString PrintType(const FDebugIds& Debug, FLooseType Type)
+{
+	FString Out;
+	Out.Reserve(64);
+	for (uint32 N = Type.NumRanges; N; --N)
+	{
+		Out.AppendChar('[');
+	}
+
+	if (Type.InnermostIsStruct())
+	{
+		Out.Append(Type.InnermostId ? Debug.Print(Type.InnermostId.Get().AsStruct()) : "struct");
+		DebugPrintVersion(Out, Type.InnermostVersion);
+		Out.Append(Type.InnermostType.Struct.IsDynamic ? "(dynamic)" : "");
+		Out.Append(Type.InnermostType.Struct.IsSuper ? "(super)" : "");
+		Out.Append(Type.InnermostType.Struct.IsDecomposed ? "(decomposed)" : "");
+	}
+	else if (Type.InnermostIsEnum())
+	{
+		Out.Append(Type.InnermostId ? Debug.Print((Type.InnermostId.Get().AsEnum())) : "enum");
+		DebugPrintVersion(Out, Type.InnermostVersion);
+		Out.Append(Type.InnermostType.Enum.FlagMode ? "(flag)" : "(flat)");
+		Out.Append(Type.InnermostType.Enum.IsDecomposed ? "(decomposed)" : "");
+	}
+	else
+	{
+		FAnsiStringView Leaves[14] = { ("Bool"), ("S8"), ("S16"), ("S32"), ("S64"), ("U8"), ("U16"), ("U32"), ("U64"), ("F32"), ("F64"), ("UTF8"), ("UTF16"), ("UTF32") }; 
+		Out.Append(Leaves[uint8(Type.InnermostType.Kind) - 2]);
+	}
+
+	if (Type.OptionalParameter)
+	{
+		Out.AppendChar('?');
+	}
+
+	for (uint32 N = Type.NumRanges; N; --N)
+	{
+		Out.AppendChar(']');
+	}
+
+	return Out;
+}
+
 ////////////////////////////////////////////////////////////////////////////////////////////////
 
 struct F64BitIndexIterator
@@ -56,14 +142,35 @@ struct FEnumeratorIdSetIterator : F64BitIndexIterator
 static FEnumeratorIdSetIterator begin(FEnumeratorIdSet Flags)		{ return {Flags.Bits}; }
 static FEnumeratorIdSetIterator end(FEnumeratorIdSet)				{ return {0}; }
 
-static F64BitIndexIterator		begin(FEnumeratorInputSet Flags)	{ return {Flags.Bits}; }
-static F64BitIndexIterator		end(FEnumeratorInputSet)			{ return {0}; }
+static F64BitIndexIterator		begin(FEnumeratorIndexSet Flags)	{ return {Flags.Bits}; }
+static F64BitIndexIterator		end(FEnumeratorIndexSet)			{ return {0}; }
 
 ////////////////////////////////////////////////////////////////////////////////////////////////
 
 namespace Upgrade
 {
 
+FString FDebugPrinter::Print(FLooseType Type) const { return PrintType(*this, Type); }
+	
+// Experimental TLS debug printer
+thread_local FDebugPrinter* TlsDebug;
+
+static FString TlsPrint(FLooseType T) { return TlsDebug->Print(T); }
+
+////////////////////////////////////////////////////////////////////////////////////////////////
+
+struct FFriend
+{
+	static const FLooseMember&		Read(const FOut& O)					{ return O; }
+	static TArrayView<FLooseMember>	Edit(TArrayView<FOut> Os)			{ return MakeArrayView(static_cast<FLooseMember*>(Os.GetData()), Os.Num()); }
+
+	template<class T>
+	static T*						NewRange(FContext Ctx, uint64 Num)	{ return Ctx.Scratch.AllocateArray<T>(Num, T{}); }
+	static FScratchAllocator&		Scratch(FContext Ctx)				{ return Ctx.Scratch; }
+};
+
+////////////////////////////////////////////////////////////////////////////////////////////////
+
 static TSet<FNameId> GetDeclaredOrder(const FEnumDeclaration& Decl)
 {
 	TSet<FNameId> Out;
@@ -114,8 +221,54 @@ static FMemoryView ViewSkippableSlice(const void* Slice)
 	return FMemoryView((const uint8*)Slice + NumBytesRead, Size);
 }
 
+static FStructView AsSourceStruct(FLooseValue Value, FSchemaBatchId Batch)
+{
+	return {{FStructSchemaId(Value.Meta.InnermostSchema.Get()), Batch}, FByteReader(ViewSkippableSlice(Value.Data.Struct.Source))};
+}
+
+class FSourceStructRangeReader
+{
+public:
+	FSourceStructRangeReader(FStructRange Structs, uint64 Num) : SourceIt(Structs.Source), NumLeft(Num) {}
+
+	bool			HasMore() const							{ return NumLeft > 0; }
+	FStructView		GrabView(FStructSchemaHandle Schema)	{ return {Schema, FByteReader(GrabSlice())}; }	
+	FStructData		GrabData()
+	{
+		FStructData Out = { .Source = static_cast<const uint8*>(SourceIt) };
+		(void)GrabSlice(); // Advance
+		return Out;
+	}
+private:
+	const void*		SourceIt;
+	uint64			NumLeft;
+
+	FMemoryView		GrabSlice()
+	{
+		check(NumLeft);
+		NumLeft--;
+		FMemoryView Out = ViewSkippableSlice(SourceIt);
+		SourceIt = Out.GetDataEnd();
+		return Out;
+	}
+};
+
 ////////////////////////////////////////////////////////////////////////////////////////////////
 
+//@pre Idx < 64
+inline void SetBit(uint64& Word, uint32 Idx, bool bValue)
+{
+	uint64 Mask = uint64(1) << Idx;
+	if (bValue)
+	{
+		Word |= Mask;
+	}
+	else
+	{
+		Word &= ~Mask;
+	}
+}
+
 static uint64 OutlinedHiBit = uint64(1) << 63;
 
 inline bool IsOutlined(FEnumeratorIdSet Flags)
@@ -138,26 +291,12 @@ FEnumeratorIdSet NewOutlined(FScratchAllocator& Scratch, uint16 NumEnumeratorIds
 // todo: this will be used by TranslateOutValue
 static void SetLooseFlag(FEnumeratorIdSet& Out, FEnumeratorId Name, bool bValue)
 {
-	uint64 Mask = uint64(1) << (Name.Idx % 64);
 	uint64& Word = IsOutlined(Out) ? AsOutlined(Out)[Name.Idx/64] : Out.Bits;
-
-	if (bValue)
-	{
-		Word |= Mask;
-	}
-	else
-	{
-		Word &= ~Mask;
-	}
+	SetBit(Word, Name.Idx % 64, bValue);
 }
 
 ////////////////////////////////////////////////////////////////////////////////////////////////
 
-inline FParameterId GetName(const FLooseMember& Member)
-{
-	return Member.GetValue().Meta.Name;
-}
-
 inline ELooseState GetInnermostState(FLooseMetadata Meta)
 {
 	if (!Meta.IsSet)
@@ -223,40 +362,78 @@ static FLooseMember AsMember(FStructMember In)
 	return FLooseMember(Type, {Meta, Data});
 }
 
+#if DO_CHECK
+
+// Opaque enums ignore flag mode and reuse FInEnum/FOutEnum[Item] for simplicity
+// ReflectItem lack type info and report flag enums as flat, this works around that
+inline bool OpaqueEnumEquals(FInnermostType Actual, FInnermostType Expected) 
+{
+	return Actual == InnermostFlagEnum && Expected == InnermostFlatEnum;
+}
+
+void CheckRangeType(FLooseType Actual, uint16 NumNamed, FInnermostType InnermostType, uint8 NumRanges)
+{
+	FLooseType Expected = {.NumRanges = NumRanges, .InnermostType = InnermostType};
+
+	checkf(Actual.NumRanges == NumRanges && (Actual.InnermostType == InnermostType || OpaqueEnumEquals(Actual.InnermostType, InnermostType)), 
+		TEXT("Parameter is a %s, not a %s"), *TlsPrint(Actual), *TlsPrint(Expected));
+	check(IsDecomposed(InnermostType) == !!NumNamed);
+	check(!IsDecomposed(InnermostType) || Actual.InnermostId);
+}
+#endif
+
 ////////////////////////////////////////////////////////////////////////////////////////////////
 
 FInStruct FIn::Struct() const
 {
-	check(Value);
 	check(Type.IsStruct());
-	return static_cast<const FInStruct&>(*this);
+	check(Type.InnermostType == InnermostStaticStruct);
+	check(Value);
+	return FInStruct(Type, Value);
 }
 
 FInStruct FIn::OptStruct() const
 {
 	check(Type.OptionalParameter);
 	check(Type.IsStruct());
-	return static_cast<const FInStruct&>(*this);
+	check(Type.InnermostType == InnermostStaticStruct);
+	return FInStruct(Type, Value);
 }
 
-const FIn* FIn::AsStructMembers(uint32 N) const
+FInAny FIn::StructAny() const
 {
 	check(Type.IsStruct());
-	check(!Value.Meta.InnermostSchema);
-	check(Value.Meta.Num == N);
+	check(Type.InnermostType == InnermostDynamicStruct);
+	check(Value);
+	return FInAny(Type, Value);
+}
+
+FInMembers FIn::StructMembers() const
+{
+	check(Type.IsStruct());
+	check(Type.InnermostType == InnermostDecomposedStruct);
+	check(Value.Meta.Num > 0);
+	check(Value);
 	check(Value.Data.Struct.Loose);
-	return static_cast<const FIn*>(Value.Data.Struct.Loose);
+	return FInMembers(Type, Value);
 }
 
-////////////////////////////////////////////////////////////////////////////////////////////////
+FInEnum FIn::Enum() const
+{
+	check(Type.IsEnum());
+	check(!Type.InnermostType.Enum.IsDecomposed);
+	check(Value.Meta.Num == 0);
+	check(!Value.Meta.InnermostSchema); // All enum are loosened since out enum ranges cant mix loose/source enums
+	check(Value);
+	return FInEnum(Type, Value.Data.Enum.Loose.Enumerator);
+}
 
-inline const FLooseMember& ReadOutput(const FOut& Out) { return Out; }
-inline TArrayView<FLooseMember> EditOutputs(TArrayView<FOut> Outs) { return MakeArrayView(static_cast<FLooseMember*>(Outs.GetData()), Outs.Num()); }
+////////////////////////////////////////////////////////////////////////////////////////////////
 
 FOut::FOut(FOutParameter Param)
 : FLooseMember(Param.Type, FLooseValue{{.Name = Param.Name, .Num = Param.NumInners}, {}})
 {
-	check(Param.NumInners == 0 || Param.Type.IsEnum());
+	check(Param.NumInners == 0 || Param.Type.InnermostType.Kind == EInnermostKind::Enum);
 }
 
 FOut::FOut(FOutParameter Param, FOut* Members)
@@ -265,65 +442,161 @@ FOut::FOut(FOutParameter Param, FOut* Members)
 	check(Param.NumInners > 0);
 }
 
-FOutEnumFlags& FOut::EnumFlags()
+FOutFlags FOut::EnumFlags()
 {
 	check(Type.IsEnum());
 	check(Type.InnermostType.Enum.FlagMode);
 	check(!Value.Meta.InnermostSchema);
 	check(Value.Meta.Num > 0 && Value.Meta.Num <= 64);
-	return *reinterpret_cast<FOutEnumFlags*>(this);
+	return FOutFlags(Type, Value);
 }
 
-FOutStruct& FOut::Struct()
+FOutStruct FOut::Struct()
 {
 	check(Type.IsStruct());
-	return *reinterpret_cast<FOutStruct*>(this);
+	check(Type.InnermostType == InnermostStaticStruct);
+	return FOutStruct(Type, Value);
 }
 
-FOut* FOut::AsStructMembers(uint32 N)
+FOutMembers FOut::StructMembers()
 {
 	check(Type.IsStruct());
+	check(Type.InnermostType == InnermostDecomposedStruct);
 	check(!Value.Meta.InnermostSchema);
-	check(Value.Meta.Num == N);
+	check(Value.Meta.Num > 0);
+	check(Value);
 	check(Value.Data.Struct.Loose);
-	return static_cast<FOut*>(Value.Data.Struct.Loose);
+	return FOutMembers(Type, Value);
 }
 
 void FOutStruct::operator=(FInStruct In)
 {
-	if (Type.InnermostType.Struct.IsDynamic)
-	{
-		Type.InnermostId = In.Type.InnermostId;
-	}
-	else
-	{
-		checkf(Type.InnermostId == In.Type.InnermostId, TEXT("Set wrong struct type"));
-	}
+	checkf(Type.InnermostId == In.GetType().InnermostId, TEXT("Set wrong struct type"));
+	checkf(Type.InnermostVersion == In.GetType().InnermostVersion, TEXT("Set wrong struct version"));
+	Value.SetNamed(In.GetValue());
+}
 
-	Value.SetNamed(In.Value);
+void FOutAny::operator=(FInStruct In)
+{
+	Owner.RecordDynamicType(In.GetType());
+	Value.SetNamed(In.GetValue());
 }
 
-void FOutEnumFlag::operator=(bool bValue)
+void FOutEnum::operator=(FInEnum In)
+{
+	check(Type.Compare(In.Type));
+	Value.Meta.IsSet = 1;
+	Value.Data.Enum.Loose.Enumerator = In.Value;
+}
+
+void FOutFlag::operator=(bool bValue)
 {
 	Out.Meta.IsSet = 1;
+	SetBit(Out.Data.Enum.Loose.Enumerator.Flags.Bits, Idx, bValue);
+}
+
+void FOutFlag::operator=(TOptional<bool> bValue)
+{
+	if (bValue)
+	{
+		operator=(*bValue);
+	}
+}
+
+////////////////////////////////////////////////////////////////////////////////////////////////
 
-	uint64& Flags = Out.Data.Enum.Loose.Indices.Flags.Bits;
+void FOutBoolItem::operator=(bool bValue)
+{
 	if (bValue)
 	{
-		Flags |= Mask;
+		Out |= Mask;
 	}
 	else
 	{
-		Flags &= ~Mask;
+		Out &= ~Mask;
+	}	
+}
+
+void FOutBoolIt::Increment()
+{
+	Mask <<= 1;
+	if (Mask == 0)
+	{
+		++Byte;
+		Mask = 1;
 	}
 }
 
-void FOutEnumFlag::operator=(TOptional<bool> bValue)
+void FOutEnumItem::operator=(FInEnum In)
 {
-	if (bValue)
+	check(Type.Compare(In.Type));
+	Out = In.Value;
+}
+
+void FOutFlagItem::operator=(bool bValue)
+{
+	SetBit(Out.Bits, Idx, bValue);
+}
+
+void FOutFlagsItem::operator=(uint64 Indices)
+{
+	check(FMath::FloorLog2_64(Indices) < Max);
+	Out.Bits = Indices;
+}
+
+////////////////////////////////////////////////////////////////////////////////////////////////
+
+void FOutRange::Set(const FInRange& In)
+{
+	check(In.NumNamed == NumNamed);
+	checkSlow(Out.GetType().Compare(In.Type));
+	Out.EditTypechecked().SetNamed(In.Value);
+}
+
+void FOutRange::Set(FLooseNestedRange In)
+{
+	FLooseValue& Value = Out.EditTypechecked();
+	Value.Meta.IsSet = 1;
+	Value.Meta.Num = In.Num;
+	Value.Data.Range = In.Data;
+}
+
+FRangeData FOutRange::NewNested(uint64 Num, FContext Ctx) const
+{
+	return { .Ranges = FFriend::NewRange<FLooseNestedRange>(Ctx, Num) };
+}
+FRangeData FOutRange::NewStructs(uint64 Num, FContext Ctx) const
+{
+	return { .Structs = {.Loose = FFriend::NewRange<FLooseValue>(Ctx, Num)} };
+}
+FRangeData FOutRange::NewMembers(uint64 Num, FContext Ctx) const
+{
+	FLooseValue* Structs = FFriend::NewRange<FLooseValue>(Ctx, Num);
+	for (FLooseValue& Struct : MakeArrayView64(Structs, Num))
 	{
-		operator=(*bValue);
+		Struct.Meta.Num = NumNamed;
 	}
+	return { .Structs = {.Loose = Structs} };
+}
+FRangeData FOutRange::NewEnums(uint64 Num, FContext Ctx) const
+{
+	return { .Enums = {.Loose = FFriend::NewRange<FEnumLooseData>(Ctx, Num)} };
+}
+FRangeData FOutRange::NewEnumerators(uint64 Num, FContext Ctx) const
+{
+	return { .Enums = {.Loose = FFriend::NewRange<FEnumLooseData>(Ctx, Num)} };
+}
+FRangeData FOutRange::NewFlags(uint64 Num, FContext Ctx) const
+{
+	return { .Enums = {.Loose = FFriend::NewRange<FEnumLooseData>(Ctx, Num)} };
+}
+FRangeData FOutRange::NewBools(uint64 Num, FContext Ctx) const
+{
+	return { .Arithmetics = {.Loose = FFriend::NewRange<uint8>(Ctx, (Num + 7) / 8)} };
+}
+FRangeData FOutRange::NewArithmetics(uint64 Num, FContext Ctx, uint32 SizeOf) const
+{
+	return { .Arithmetics = {.Loose = FFriend::Scratch(Ctx).AllocateZeroed(Num * SizeOf, SizeOf)} };
 }
 
 ////////////////////////////////////////////////////////////////////////////////////////////////
@@ -340,9 +613,9 @@ const FEnumLoosener& FEnumLooseners::operator[](FEnumSchemaId Id) const
 	return Data[Id.Idx - NumStructs].Source;
 }
 
-FEnumLooseners ViewLooseners(FBatchContext& Ctx)
+static FEnumLooseners ViewLooseners(FBatchContext& Ctx)
 {
-	return FEnumLooseners(Ctx.Transformations.Enums, Ctx.Transformations.NumSourceStructs);
+	return FEnumLooseners(Ctx.Transformations.Enums, Ctx.Transformations.Schemas.NumStructs);
 }
 
 ////////////////////////////////////////////////////////////////////////////////////////////////
@@ -364,6 +637,11 @@ static FEnumeratorData Loosen(const FFlagEnumLoosener& Loosener, uint64 Flags)
 	return { .Flags = Out };
 }
 
+static FEnumeratorData Loosen(const FEnumLoosener& Loosener, uint64 Constant)
+{
+	return Loosener.IsFlag() ? Loosen(Loosener.AsFlag(), Constant) : Loosen(Loosener.AsFlat(), Constant);
+}
+
 static uint64 Fasten(TConstArrayView<FEnumerator> Enumerators, FEnumeratorId Flat, const FEnumeratorIndexer& Indexer)
 {
 	FNameId Name = Indexer.Resolve(Flat);
@@ -447,12 +725,12 @@ FLooseValue FLooseMemberReader::Grab()
 			FEnumeratorData Data = Next.Type.InnermostType.Enum.FlagMode 
 								 ? Loosen(Loosener.AsFlag(), Value)
 								 : Loosen(Loosener.AsFlat(), Value);
-			return { Meta, FLooseData{ .Enum = { .Loose = { .Enumerators = Data} } } };		
+			return { Meta, {.Enum = {.Loose = {.Enumerator = Data }}} };		
 		}
 		else // source schema matches declaration
 		{
 			Meta.InnermostSchema = ToOptionalSchema(Leaf.Enum);
-			return { Meta, FLooseData{ .Enum = {.Source = Value} } };	
+			return { Meta, {.Enum = {.Source = Value}} };	
 		}
 	}
 	else
@@ -474,31 +752,49 @@ FRangeMember FLooseMemberReader::ToLooseRange(FParameterId Name, FLooseType Type
 	{
 		return { Type, Meta, {.Structs = {.Source = GetValues(Range).GetData()}} };
 	}
+	else if (Type.InnermostType.Kind == EInnermostKind::Enum)
+	{
+		return { Type, Meta, {.Enums = {.Source = GetValues(Range).GetData()}} };
+	}
 	else
 	{
-		return { Type, Meta, {.Leaves = GetValues(Range).GetData()} };
+		return { Type, Meta, {.Arithmetics = {.Source = GetValues(Range).GetData()}} };
 	}
 }
 
-const FLooseNestedRange* FLooseMemberReader::LoosenRanges(FNestedRangeView Ins)
+FLooseNestedRange* FLooseMemberReader::LoosenRanges(FNestedRangeView Ins)
 {
+	checkSlow(Ins.GetSchema().ItemType.IsRange());
 	FLooseNestedRange* Outs = Scratch.AllocateArray<FLooseNestedRange>(Ins.Num());
 	FLooseNestedRange* OutIt = Outs;
-	EMemberKind ItemKind = Ins.GetSchema().ItemType.GetKind();
-	if (ItemKind == EMemberKind::Range)
+	FMemberType ItemType = Ins.GetSchema().NestedItemTypes[0];
+	bool bEnum = IsEnum(ItemType);
+	if (ItemType.IsRange())
 	{
 		for (FRangeView In : Ins)
 		{
 			*OutIt++ = { In.Num(), {.Ranges = LoosenRanges(In.AsRanges())} };
 		}
 	}
+	else if (ItemType.IsStruct())
+	{
+		for (FRangeView In : Ins)
+		{
+			*OutIt++ = { In.Num(), {.Structs = {.Source = GetValues(In).GetData()}} };
+		}
+	}
+	else if (bEnum)
+	{
+		for (FRangeView In : Ins)
+		{
+			*OutIt++ = { In.Num(), {.Enums = {.Source = GetValues(In).GetData()}} };
+		}
+	}
 	else
 	{
-		bool bStruct = ItemKind == EMemberKind::Struct;
 		for (FRangeView In : Ins)
 		{
-			*OutIt++ = { In.Num(), bStruct ? FRangeData{ .Structs = {.Source = GetValues(In).GetData()} }
-										   : FRangeData{ .Leaves  = GetValues(In).GetData()} };
+			*OutIt++ = { In.Num(), {.Arithmetics = {.Source = GetValues(In).GetData()}} };
 		}
 	}
 
@@ -532,7 +828,6 @@ FOutputStack::FOutputStack(FScratchAllocator& InScratch)
 , Scratch(InScratch)
 {}
 
-
 static FOut AllocateStructMembers(FOutParameter Param, TConstArrayView<FOutParameter>& InOutParams, FScratchAllocator& Scratch)
 {
 	uint32 Num = Param.NumInners;
@@ -583,6 +878,20 @@ FOutputScope::FOutputScope(FOutputStack& InStack, TConstArrayView<FOutParameter>
 , Outputs(Stack.Push(Params))
 {}
 
+//////////////////////////////////////////////////////////////////////////////////////////////////
+//
+//struct FOutsIt
+//{
+//	friend inline FOutsIt begin(FOuts Outs) { return {Outs.begin(), Outs.Num}; }
+//	friend inline FOutsIt end(FOuts Outs)	{ return {Outs.end(), Outs.Num}; }
+//
+//	FLooseType*		It;
+//	const uint32	Num;
+//
+//	FOutItem operator*() const { return {*It, *reinterpret_cast<FLooseValue*>(It+Num); }
+//};
+
+
 ////////////////////////////////////////////////////////////////////////////////////////////////
 
 inline TConstArrayView<const FVersionUpgrades*> ViewVersions(const FStructUpgrades& Struct)
@@ -756,6 +1065,100 @@ FORCENOINLINE static FLooseMember UpgradeUnmatchedMember(FLooseType Type, FLoose
 	return AsMember(Out);
 }
 
+// Source range of source structs -> loose range of source structs
+static FStructRange LoosenStructRangeOnly(FLooseNestedRange Range, FSchemaId Schema, FScratchAllocator& Scratch)
+{
+	FLooseValue* Out = Scratch.AllocateArray<FLooseValue>(Range.Num, FLooseValue{{.IsSet = 1, .InnermostSchema = Schema}});
+	FSourceStructRangeReader SourceIt(Range.Data.Structs, Range.Num);
+	for (uint64 Idx = 0; Idx < Range.Num; ++Idx)
+	{
+		Out[Idx].Data.Struct = SourceIt.GrabData();
+	}
+	return { .Loose = Out };
+}
+
+template<typename T, typename LoosenerType>
+static void LoosenEnums(FEnumLooseData* Dst, TConstArrayView64<T> Src, const LoosenerType& Loosener)
+{
+	for (T Constant : Src)
+	{
+		*Dst++ = { .Enumerator = Loosen(Loosener, Constant) };
+	}
+}
+
+template<typename T>
+static void LoosenEnums(FEnumLooseData* Dst, const T* Src, uint64 Num, const FEnumLoosener& Loosener)
+{
+	return Loosener.IsFlag()	? LoosenEnums<T>(Dst, MakeArrayView64(Src, Num), Loosener.AsFlag())
+								: LoosenEnums<T>(Dst, MakeArrayView64(Src, Num), Loosener.AsFlat());
+}
+
+static FEnumRange LoosenEnums(ELeafWidth SourceWidth, FLooseNestedRange Range, const FEnumLoosener& Loosener, FScratchAllocator& Scratch)
+{
+	check(Loosener.Exists());
+	const void* Src = Range.Data.Enums.Source;
+	FEnumLooseData* Dst = Scratch.AllocateArray<FEnumLooseData>(Range.Num);
+	switch (SourceWidth)
+	{
+		case ELeafWidth::B8:	LoosenEnums(Dst, static_cast<const uint8* >(Src), Range.Num, Loosener); break;
+		case ELeafWidth::B16:	LoosenEnums(Dst, static_cast<const uint16*>(Src), Range.Num, Loosener); break;
+		case ELeafWidth::B32:	LoosenEnums(Dst, static_cast<const uint32*>(Src), Range.Num, Loosener); break;
+		case ELeafWidth::B64:	LoosenEnums(Dst, static_cast<const uint64*>(Src), Range.Num, Loosener); break;
+		default:				check(false); break;
+	}
+
+	return { .Loose = Dst };
+}
+
+static FRangeData LoosenRangeForInput(FLooseType Type, FLooseNestedRange Value, FSchemaId SchemaId, FBatchContext& Ctx)
+{
+	check(Value.Num);
+
+	FLooseType ItemType = ItemOf(Type);
+	if (ItemType.IsRange())
+	{
+		for (FLooseNestedRange& Range : MakeArrayView64(Value.Data.Ranges, Value.Num))
+		{
+			Range.Data = Range.Num ? LoosenRangeForInput(ItemType, Range, SchemaId, Ctx) : Range.Data;
+		}
+		return Value.Data;
+	}
+	else if (ItemType.IsStruct())
+	{
+		check(!Type.InnermostType.Struct.IsDecomposed);
+		return { .Structs = LoosenStructRangeOnly(Value, SchemaId, Ctx.Scratch) };
+	}
+	else
+	{
+		check(ItemType.IsEnum());
+		check(!Type.InnermostType.Enum.IsDecomposed);
+		FEnumSchemaId Id(SchemaId);
+		const FEnumSchema& Schema = ResolveEnumSchema(Ctx.Batch, Id);
+		return { .Enums = LoosenEnums(Schema.Width, Value, ViewLooseners(Ctx)[Id], Ctx.Scratch) };
+	}
+}
+
+static FLooseValue LoosenForInput(FLooseType Type, FLooseValue Value, FSchemaId SchemaId, FBatchContext& Ctx)
+{
+	if (Type.IsRange())
+	{
+		check(Type.NumRanges > 1 || Value.Meta.IsSourceRange);
+		if (uint64 Num = Value.Meta.Num)
+		{
+			Value.Data.Range = LoosenRangeForInput(Type, {Value.Meta.Num, Value.Data.Range}, SchemaId, Ctx);
+		}
+	}
+	else
+	{
+		check(Type.IsEnum());
+		Value.Data.Enum.Loose.Enumerator = Loosen(ViewLooseners(Ctx)[FEnumSchemaId(SchemaId)], Value.Data.Enum.Source);
+	}
+
+	Value.Meta.InnermostSchema = NoId;
+	Value.Meta.IsSourceRange = 0;
+	return Value;
+}
+				
 void FInstanceUpgrader::Give(FLooseType Type, FLooseValue Value, FMatchHelper Helper)
 {
 	FParameterId Name = Value.Meta.Name;
@@ -783,9 +1186,13 @@ Rematch:
 					checkf(ScanIdx < NameMatches.Num, TEXT("Input decomposition match can't be last match for '%s' in '%s'"), *Print(Name), *Print(Struct.Id));
 					Match = Helper.AllMatches[NameMatches.Idx + ScanIdx];
 					++ScanIdx;
-					checkf(Match.Type.Compare(Type), TEXT("The match after an input decomposition match should be of same type"));
+					checkf(CompareDecomposedAndRecomposed(Match.Type, Type), TEXT("The match after an input decomposition match should be of same type, expected %s, actual %s"), *Print(Type), *Print(Match.Type));
 					checkf(!Match.Decomposes(), TEXT("Can't input decompose twice"));
 				}
+				else if (Value.Meta.InnermostSchema && (Type.InnermostIsEnum() || Type.IsRange()))
+				{
+					Value = LoosenForInput(Type, Value, Value.Meta.InnermostSchema.Get(), Ctx);
+				}
 
 				Input(Match, Value);
 				return;
@@ -799,7 +1206,7 @@ Rematch:
 			// Todo: Only upgrade to next matching version, e.g. upgrading V1 -> V2 -> V3 would miss a V2 match.
 			//		 Determine next version by looking at matches
 			FLooseMember Upgraded = UpgradeUnmatchedMember(Type, Value, *Xform, Ctx);
-			checkSlow(Name.Idx == GetName(Upgraded).Idx);
+			checkSlow(Name.Idx == Upgraded.GetName().Idx);
 
 			bool bVersionUpdated = !Upgraded.GetType().Compare(Type);
 			Type = Upgraded.GetType();
@@ -897,16 +1304,114 @@ void FInstanceUpgrader::Input(FMemberMatch Match, FLooseValue Value)
 	}
 }
 
+static FLooseValue LoosenStruct(const FLooseSchema& LooseSchema, FStructView Source, FEnumLooseners Looseners, FScratchAllocator& Scratch)
+{
+	TArray<FLooseMember, TInlineAllocator<16>> Members;
+	FLooseMemberReader It(LooseSchema, Source, Looseners, Scratch);
+	while (It.HasMore())
+	{
+		FLooseParameter Param = It.Peek(); // Peek() before Grab()
+		Members.Emplace(Param.Type, It.Grab());
+	}
+	
+	uint64 Num(Members.Num());
+	return {{.IsSet = 1, .Num = Num }, {.Struct = {.Loose = Scratch.CloneArray(Num, Members.GetData())}}};
+}
+
+static const FLooseSchema& ObtainSourceSchema(FBatchContext& Ctx, FDeclId Id, FStructSchemaId SchemaId);
+
+static FRangeData LoosenItemsForDecomposition(FLooseType ItemType, FSchemaId SchemaId, FLooseNestedRange Range, FBatchContext& Ctx)
+{
+	if (ItemType.IsRange())
+	{
+		FLooseType InnerItemType = ItemOf(ItemType);
+		for (FLooseNestedRange Item : MakeArrayView64(Range.Data.Ranges, Range.Num))
+		{
+			LoosenItemsForDecomposition(InnerItemType, SchemaId, Item, Ctx);
+		}
+		return {.Ranges = Range.Data.Ranges};
+	}
+	else if (ItemType.IsStruct())
+	{
+		FDeclId Id = ItemType.InnermostId.Get().AsStructDeclId();
+		const FLooseSchema& LooseSchema = ObtainSourceSchema(Ctx, Id, FStructSchemaId(SchemaId));
+		FEnumLooseners Looseners = ViewLooseners(Ctx);
+
+		FLooseValue* Out = Ctx.Scratch.AllocateArray<FLooseValue>(Range.Num);
+		FStructSchemaHandle Schema = { FStructSchemaId(SchemaId), Ctx.Batch };
+		FSourceStructRangeReader SourceIt(Range.Data.Structs, Range.Num);
+		for (uint64 Idx = 0; Idx < Range.Num; ++Idx)
+		{
+			Out[Idx] = LoosenStruct(LooseSchema, SourceIt.GrabView(Schema), Looseners, Ctx.Scratch);
+		}
+
+		return {.Structs = {.Loose = Out}};
+	}
+	else
+	{
+		check(ItemType.IsEnum());
+		unimplemented();
+		FEnumLooseData* Out = nullptr;
+
+		return {.Enums = {.Loose = Out}};
+	}
+}
+
 FLooseValue FInstanceUpgrader::Decompose(FLooseType Type, FLooseValue Value, uint16 DecompIdx)
 {
 	check(DecompIdx < Current->NumDecompositions);
 	uint32 Offset = Current->GetDecompOffsets()[DecompIdx];
 	const void* Ptr = ((const uint8*)Current) + Offset;
+	if (Type.IsRange())
+	{
+		FLooseType ItemType = ItemOf(Type);
+		if (FOptionalSchemaId Schema = Value.Meta.InnermostSchema)
+		{
+			Value.Data.Range = LoosenItemsForDecomposition(ItemType, Schema.Get(), {Value.Meta.Num, Value.Data.Range}, Ctx);
+			Value.Meta.InnermostSchema = NoId;
+		}
+
+		Decompose(ItemType, {Value.Meta.Num, Value.Data.Range}, Ptr);
+		return Value;
+	}
+	
 	return Type.IsStruct()
 		? Decompose(AsStruct(Type, Value.Meta, Value.Data.Struct), *static_cast<const FStructDecomposition*>(Ptr))
 		: Decompose(AsEnum(Type, Value.Meta, Value.Data.Enum), *static_cast<const FEnumDecomposition*>(Ptr));
 }
 
+void FInstanceUpgrader::Decompose(FLooseType ItemType, FLooseNestedRange Range, const void* Decomp)
+{
+	if (Range.Num == 0)
+	{
+		return;
+	}
+	else if (ItemType.IsRange())
+	{
+		for (FLooseNestedRange Item : MakeArrayView64(Range.Data.Ranges, Range.Num))
+		{
+			Decompose(ItemOf(ItemType), Item, Decomp);
+		}
+	}
+	else if (ItemType.IsStruct())
+	{
+		const FStructDecomposition& StructDecomp = *static_cast<const FStructDecomposition*>(Decomp);
+		for (FLooseValue& Item : MakeArrayView64(Range.Data.Structs.Loose, Range.Num))
+		{
+			Item = Decompose(AsStruct(ItemType, Item.Meta, Item.Data.Struct), StructDecomp);
+		}
+	}
+	else
+	{
+		const FEnumDecomposition& EnumDecomp = *static_cast<const FEnumDecomposition*>(Decomp);
+		const FLooseMetadata Dummy = {};
+		for (FEnumLooseData& Item : MakeArrayView64(Range.Data.Enums.Loose, Range.Num))
+		{
+			Item = Decompose(AsEnum(ItemType, Dummy, {.Loose = Item}), EnumDecomp).Data.Enum.Loose;
+		}
+	}
+}
+
 FLooseValue FInstanceUpgrader::Decompose(FLooseType Type, FLooseValue Value, const FEnumDecomposition& Decomp)
 {
 	check(Value.Meta.IsSet);
@@ -915,23 +1420,22 @@ FLooseValue FInstanceUpgrader::Decompose(FLooseType Type, FLooseValue Value, con
 
 FLooseValue FInstanceUpgrader::Decompose(FEnumMember In, const FEnumDecomposition& Decomp)
 {
-	FEnumeratorData Enumerators = {};
+	FEnumeratorData Enumerator = {};
 	if (In.State == ELooseState::Source)
 	{
 		FEnumLoosener Loosener = Ctx.Transformations[In.Schema.Get()].Source;
-		Enumerators = In.FlagMode	? Loosen(Loosener.AsFlag(), In.Data.Source)
+		Enumerator = In.FlagMode	? Loosen(Loosener.AsFlag(), In.Data.Source)
 									: Loosen(Loosener.AsFlat(), In.Data.Source);
 	}
 	else
 	{
-		Enumerators = In.Data.Loose.Enumerators;
+		Enumerator = In.Data.Loose.Enumerator;
 	}
 
-	FEnumIndexData Indices = In.FlagMode	? Reorder(Enumerators.Flags, Decomp.InputOrder) 
-											: Reorder(Enumerators.Flat, Decomp.InputOrder);
-
+	FEnumIndexData Decomposed = In.FlagMode	? Reorder(Enumerator.Flags, Decomp.InputOrder) 
+											: Reorder(Enumerator.Flat, Decomp.InputOrder);
 	FLooseMetadata Meta = { .IsSet = 1, .Name = In.Name };
-	return { Meta, { .Enum = { .Loose = { .Indices = Indices } } } };
+	return { Meta, {.Enum = {.Loose = {.Decomposed = Decomposed}}} };
 }
 
 static int32 GetInputIndex(FEnumeratorId Name, const TSet<FEnumeratorId>& Order)
@@ -1043,7 +1547,7 @@ static TArrayView<FLooseMember> ReorderInputs(TConstArrayView<FLooseMember> Unso
 	{
 		if (Src)
 		{
-			FSetElementId Place = Order.FindId(GetName(Src));
+			FSetElementId Place = Order.FindId(Src.GetName());
 			AddInput(Place.IsValidId() ? &Expected[Place.AsInteger()] : nullptr, Src.GetType(), Src.GetValue(), /* out */ Unexpected);
 		}
 	}
@@ -1147,9 +1651,28 @@ FOuter FInstanceUpgrader::Resolve(uint8 Recipient)
 	return It->Outer;
 }
 
+void FInstanceUpgrader::Recompose(FRangeData Out, uint64 Num, uint8 Depth, TConstArrayView<FEnumeratorId> Order, FInnermostEnumType Type)
+{
+	if (Depth > 1)
+	{
+		for (FLooseNestedRange Inner : TConstArrayView64<FLooseNestedRange>(Out.Ranges, Num))
+		{
+			Recompose(Inner.Data, Inner.Num, Depth - 1, Order, Type);
+		}
+	}
+	else
+	{
+		for (FEnumLooseData& Enum : TArrayView64<FEnumLooseData>(Out.Enums.Loose, Num))
+		{
+			Enum.Enumerator = Recompose(Enum.Decomposed, Order, Type);
+		}
+	}
+}
+
 // Opt: avoid creating "no-op" recompositions if IsDeclarationOrder(Order)
-FEnumeratorData FInstanceUpgrader::Recompose(FInnermostEnumType Type, FEnumIndexData In, TConstArrayView<FEnumeratorId> Order)
+FEnumeratorData FInstanceUpgrader::Recompose(FEnumIndexData In, TConstArrayView<FEnumeratorId> Order, FInnermostEnumType Type)
 {
+	check(Type.IsDecomposed);
 	if (Type.FlagMode)
 	{
 		FEnumeratorIdSet Out = {};
@@ -1174,23 +1697,41 @@ static TConstArrayView<FEnumeratorId> GetRecompOrder(const FVersionUpgrades& Ver
 	return MakeArrayView(reinterpret_cast<const FEnumeratorId*>(&Jump + Jump), Output.NumInners);
 }
 
-void FInstanceUpgrader::Recompose(TArrayView<FOut> Outputs, TConstArrayView<FOutParameter> Params)
+void FInstanceUpgrader::Recompose(FOuts Outputs, TConstArrayView<FOutParameter> Params)
 {
 	checkSlow(Outputs.Num() <= Params.Num());
-	for (FLooseMember& Output : EditOutputs(Outputs))
+	for (FLooseMember& Output : FFriend::Edit(Outputs))
 	{
 		FOutParameter Param = PopFirst(Params);
 		
 		check(Output.GetType().CompareAll(Param.Type));
-		checkf(Output || Param.Type.OptionalParameter, TEXT("Out parameter '%s' is required but wasn't set while upgrading %s"), 
-			*Print(GetName(Output)), *Ctx.Debug.Print(Struct.Id));
+		checkf(Output || Param.Type.OptionalParameter, TEXT("Required output '%s' missing while upgrading %s"), 
+			*Print(Output.GetName()), *Ctx.Debug.Print(Struct.Id));
 		
-		if (Param.NumInners && Param.Type.InnermostType.Kind == EInnermostKind::Enum)
+		if (Param.NumInners)
 		{
-			check(Output.GetValue().Meta.Num == Param.NumInners);
-			Output.EditTypechecked().Meta.Num = 0; // Reset num enumerator indices
-			FEnumLooseData& Enum = Output.EditTypechecked().Data.Enum.Loose;
-			Enum.Enumerators = Recompose(Param.Type.InnermostType.Enum, Enum.Indices, GetRecompOrder(*Current, Param));
+			if (Param.Type.InnermostType.Kind == EInnermostKind::Enum)
+			{
+				TConstArrayView<FEnumeratorId> Order = GetRecompOrder(*Current, Param);
+				FLooseValue& OutValue = Output.EditTypechecked();
+				if (Param.Type.IsRange())
+				{
+					Recompose(OutValue.Data.Range, OutValue.Meta.Num, Param.Type.NumRanges, Order, Param.Type.InnermostType.Enum);						
+				}
+				else
+				{
+					check(OutValue.Meta.Num == Param.NumInners);
+					OutValue.Meta.Num = 0; // Reset num enumerator indices
+					FEnumLooseData& Enum = OutValue.Data.Enum.Loose;
+					Enum.Enumerator = Recompose(Enum.Decomposed, Order, Param.Type.InnermostType.Enum);			
+				}
+				Output.MarkEnumRecomposed();
+			}
+			else
+			{
+				check(Param.Type.InnermostType.Kind == EInnermostKind::Struct);
+				Output.MarkStructRecomposed();
+			}
 		}
 	}
 }
@@ -1210,7 +1751,7 @@ void FInstanceUpgrader::Run(FOp Op, const FIn* Inputs, uint8 Recipient)
 		FMatchHelper Helper(*To.Instance->Current);
 		for (const FOut& Output : Scope.Outputs)
 		{
-			if (const FLooseMember& Member = ReadOutput(Output))
+			if (const FLooseMember& Member = FFriend::Read(Output))
 			{
 				To.Instance->Give(Member.GetType(), Member.GetValue(), Helper);
 			}
@@ -1220,7 +1761,7 @@ void FInstanceUpgrader::Run(FOp Op, const FIn* Inputs, uint8 Recipient)
 	{
 		for (const FOut& Output : Scope.Outputs)
 		{
-			if (const FLooseMember& Member = ReadOutput(Output))
+			if (const FLooseMember& Member = FFriend::Read(Output))
 			{
 				Give(Member.GetType(), Member.GetValue());
 			}
@@ -1268,7 +1809,7 @@ static void RenameMembers(TArrayView<FLooseMember> Members, const FParameterInde
 	{
 		if (Member || Member.IsMismatch())
 		{
-			FOptionalMemberId Name = From.Resolve(GetName(Member));
+			FOptionalMemberId Name = From.Resolve(Member.GetName());
 			FParameterId Reindexed = To.Index(Name);
 			Member.Rename(Reindexed);
 		}
@@ -1291,7 +1832,7 @@ FLooseStruct FInstanceUpgrader::Finish(EMismatch& OutMismatch)
 	for (const FLooseMember& Member : Outcome)
 	{
 		if (!Member && !Member.IsOptional() &&
-			!Algo::FindByPredicate(Mismatches, [&Member](const FLooseMember& Mismatch) { return Member.GetValue().Meta.Name == Mismatch.GetValue().Meta.Name; }))
+			!Algo::FindByPredicate(Mismatches, [&Member](const FLooseMember& Mismatch) { return Member.GetName() == Mismatch.GetName(); }))
 		{
 			Mismatches.Add_GetRef(Member).EditTypechecked().Meta.Mismatch = EMismatch::Missing;
 		}
@@ -1311,6 +1852,11 @@ FString FInstanceUpgrader::Print(FParameterId Parameter) const
 	return Ctx.Debug.Print(Ctx.LooseIds.Get(Struct.Id).Resolve(Parameter));
 }
 
+FString FInstanceUpgrader::Print(FLooseType Type) const
+{
+	return Ctx.Debug.Print(Type);
+}
+
 FLooseStruct FInstanceUpgrader::Upgrade(FStructView In, const FLooseSchema& From, const FLooseSchema* To, EMismatch& OutMismatch)
 {
 	check(Current == GetVersion(Struct, From.Version));
@@ -1492,9 +2038,9 @@ static bool	MatchStruct(const FStructDeclaration& Decl, const FStructSchema& Sch
 
 struct FNestedMemberType
 {
-	const FMemberType*	InnerRangeTypes = nullptr;
-	uint32				NumInnerRanges = 0;
 	FMemberType			Type;
+	uint32				NumInnerRanges = 0;
+	const FMemberType*	InnerRangeTypes = nullptr;
 };
 
 // CreateSourceSchema/CreateTargetSchema/FastenStruct helper
@@ -1520,17 +2066,19 @@ struct FMemberTypeIterator
 
 	FNestedMemberType Grab()
 	{
-		FNestedMemberType Out;
 		check(TypeIt < TypeEnd);
-		Out.Type = *TypeIt++;
-		Out.InnerRangeTypes = Out.Type.IsRange() ? RangeIt : nullptr;
-		for (FMemberType Type = Out.Type; Type.IsRange(); Type = *RangeIt++)
+		FMemberType Type = *TypeIt++;
+		return Type.IsRange() ? GrabInnerRanges(Type) : FNestedMemberType{Type};
+	}
+
+	FNestedMemberType GrabInnerRanges(FMemberType Type)
+	{
+		const FMemberType* Out = RangeIt;
+		while ((RangeIt++)->IsRange())
 		{
-			++Out.NumInnerRanges;
-			check(RangeIt < RangeEnd);
+			check(RangeIt <= RangeEnd);
 		}
-		
-		return Out;
+		return {Type, static_cast<uint32>(RangeIt - Out), Out};
 	}
 };
 
@@ -1564,7 +2112,7 @@ struct FSchemaInnermostIdIterator
 
 	FInnerId GrabEnum(EEnumMode& OutMode)
 	{
-		OutMode = EnumFlagModes[It->Idx] ? EEnumMode::Flag : EEnumMode::Flat;
+		OutMode = EnumFlagModes[It->Idx - StructVersions.Num()] ? EEnumMode::Flag : EEnumMode::Flat;
 		FSchemaId Id = *It++;
 		return RuntimeIds[Id];
 	}
@@ -1643,6 +2191,10 @@ FInnermostType GrabInnermost(InnermostIterator& InnermostIt, FMemberType In, FOp
 // CreateSourceSchema helper
 struct FSourceLooseTypeIterator
 {
+	FSourceLooseTypeIterator(const FStructSchema& Struct, const FSourceSchemas& Schemas)
+	: FSourceLooseTypeIterator(Struct, Schemas.RuntimeIds, Schemas.StructVersions, Schemas.EnumFlagModes)
+	{}
+
 	FSourceLooseTypeIterator(const FStructSchema& Struct, FInnerIds RuntimeIds, TConstArrayView<uint16> Versions, const TBitArray<>& FlagModes)
 	: SkipSuper(SkipDeclaredSuperSchema(Struct.Inheritance))
 	, TypeIt(Struct.GetMemberTypes().LeftChop(SkipSuper), Struct.GetRangeTypes())
@@ -1686,7 +2238,8 @@ struct FTargetLooseTypeIterator
 	{
 		FNestedMemberType Mid = TypeIt.Grab();
 		FMemberType InnerRangeType = Mid.InnerRangeTypes ? Mid.InnerRangeTypes[0] : Mid.Type;
-		FOptionalInnerId InnermostId = IsStructOrEnum(Mid.Type) ? ToOptional(InnerIdIt.Grab()) : NoId;
+		FMemberType InnermostType = Mid.NumInnerRanges ? Mid.InnerRangeTypes[Mid.NumInnerRanges - 1] : Mid.Type;
+		FOptionalInnerId InnermostId = IsStructOrEnum(InnermostType) ? ToOptional(InnerIdIt.Grab()) : NoId;
 		return { Mid.Type, InnerRangeType, static_cast<uint16>(Mid.NumInnerRanges), InnermostId, Mid.InnerRangeTypes };
 	}
 };
@@ -1698,17 +2251,22 @@ struct FBatchRun
 	: RuntimeIds{{InnerIds}}
 	, History(Upgrades)
 	, Scratch(InScratch)
-	, Debug(Upgrades.Ids)
-	{}
+	, Debug(Upgrades.Ids, History.LooseIds)
+	{
+		TlsDebug = &Debug;
+	}
+	~FBatchRun()
+	{
+		TlsDebug = nullptr;
+	}
 
 	FInnerIds						RuntimeIds;
 	FTransformations				Transformations;
 	const FHistory&					History;
 	FScratchAllocator&				Scratch;
 	FSchemaBatchId					SourceBatch;
-	const FDebugIds					Debug;
+	FDebugPrinter					Debug;
 
-	FLooseSchema*					CreateSourceSchema(const FStructSchema& Struct, FParameterIndexer& ParamIds, TConstArrayView<uint16> StructVersions, const TBitArray<>& EnumFlagModes);
 	void							CreateTransformations(const TBitArray<>& MatchingSchemas, const FSchemaBatch& Batch);
 	FUpgradedInstance				UpgradeOrRebuild(FStructView Source);
 	FUpgradedInstance				UpgradeStruct(FStructView Source, FTransformation Xform) const;
@@ -1745,18 +2303,16 @@ static FLooseSchema* CreateSchema(FScratchAllocator& Scratch, uint16 Version, ui
 	return new (Scratch.Allocate(Size, alignof(FLooseSchema))) FLooseSchema(Header);
 }
 
-FLooseSchema* FBatchRun::CreateSourceSchema(const FStructSchema& Schema, FParameterIndexer& ParamIds, TConstArrayView<uint16> StructVersions, const TBitArray<>& EnumFlagModes)
+static FLooseSchema* CreateSourceSchema(const FStructSchema& Schema, const FSourceSchemas& Schemas, FParameterIndexer& ParamIds, FScratchAllocator& Scratch)
 {
 	FLooseSchema* Out = CreateSchema(Scratch, Schema.Version, Schema.NumMembers);
-
 	const FMemberId* NameIt = Schema.GetMemberNames().GetData();
-	FSourceLooseTypeIterator TypeIt(Schema, RuntimeIds, StructVersions, EnumFlagModes);
+	FSourceLooseTypeIterator TypeIt(Schema, Schemas);
 	for (FLooseParameter& Param : MakeArrayView(Out->Members, Schema.NumNames()))
 	{
 		Param.Name = ParamIds.Index(*NameIt++);
 		Param.Type = TypeIt.Grab();
 	}
-
 	return Out;
 }
 
@@ -1810,6 +2366,22 @@ static TBitArray<> ExtractFlagModes(TSchemaRange<const FEnumSchema> Schemas)
 	return Out;
 }
 
+static void AddEnumInputSchemas(TBitArray<>& OutEnumSchemaInputs, const TSet<FEnumId>& EnumLookup, const FVersionUpgrades& Version)
+{
+	for (FMemberMatch Match : MakeArrayView(Version.GetMatches(), Version.NumMatches))
+	{
+		if (Match.Type.InnermostIsEnum())
+		{
+			FEnumId Id = Match.Type.InnermostId.Get().AsEnum();
+			FSetElementId Idx = EnumLookup.FindId(Id);
+			if (Idx.IsValidId() && Idx.AsInteger() < OutEnumSchemaInputs.Num())
+			{
+				OutEnumSchemaInputs[Idx.AsInteger()] = true;
+			}
+		}
+	}
+}
+
 static void AddInnerIdOutputs(TSet<FDeclId>& OutStructs, TSet<FEnumId>& OutEnums, const FVersionUpgrades& Version)
 {
 	const FOutParameter* Outputs = Version.GetOutputs();
@@ -1820,12 +2392,13 @@ static void AddInnerIdOutputs(TSet<FDeclId>& OutStructs, TSet<FEnumId>& OutEnums
 			if (Output.Type.InnermostId)
 			{
 				FInnerId Inner = Output.Type.InnermostId.Get();
-				if (Output.Type.IsStruct())
+				if (Output.Type.InnermostIsStruct())
 				{
 					OutStructs.Emplace(Inner.AsStruct());
 				}
 				else
 				{
+					check(Output.Type.InnermostIsEnum());
 					OutEnums.Add(Inner.AsEnum());
 				}
 			}
@@ -1833,10 +2406,63 @@ static void AddInnerIdOutputs(TSet<FDeclId>& OutStructs, TSet<FEnumId>& OutEnums
 	}
 }
 
-static FEnumLoosener CreateEnumLoosener(const FEnumSchema& Schema, bool bNeeded, FEnumeratorIndexer& Ids, FScratchAllocator& Scratch)
+static bool LoosenerRequired(const FFlagEnumLoosener& Loosener, uint32 Num)
+{
+	for (uint32 Idx = 0; Idx < Num; ++Idx)
+	{
+		if (Loosener.FlagNames[Idx].Idx != Idx)
+		{
+			return true;
+		}
+	}
+	return false;
+}
+
+template<typename T>
+static FEnumLoosener CreateFlagLoosener(const FEnumSchema& Schema, bool bNeeded, FEnumeratorIndexer& Ids, FScratchAllocator& Scratch)
+{
+	FFlagEnumLoosener Out = {};
+	if (const T* Constants = Schema.ExplicitConstants ? GetConstants<T>(Schema).GetData() : nullptr)
+	{
+		for (uint32 Idx = 0, Num = Schema.Num; Idx < Num; ++Idx)
+		{
+			check(FMath::CountBits64(Constants[Idx]) == 1);
+			uint64 BitIdx = FMath::FloorLog2NonZero_64(Constants[Idx]);
+			Out.FlagNames[BitIdx] = Ids.Index(Schema.Footer[Idx]);
+		}
+	}
+	else
+	{
+		for (uint32 Idx = 0, Num = Schema.Num; Idx < Num; ++Idx)
+		{
+			Out.FlagNames[Idx] = Ids.Index(Schema.Footer[Idx]);
+		}
+	}
+
+	return LoosenerRequired(Out, Schema.Num) ? FEnumLoosener(bNeeded, *new (Scratch) FFlagEnumLoosener(Out)) : FEnumLoosener();
+}
+
+template<typename T>
+static FEnumLoosener CreateFlatLoosener(const FEnumSchema& Schema, bool bNeeded, FEnumeratorIndexer& Ids, FScratchAllocator& Scratch)
 {
 	unimplemented();
-	return {};
+	return FEnumLoosener();
+}
+
+static FEnumLoosener CreateEnumLoosener(const FEnumSchema& Schema, bool bNeeded, FEnumeratorIndexer& Ids, FScratchAllocator& Scratch)
+{
+	switch (Schema.Width)
+	{
+		case ELeafWidth::B8:	return Schema.FlagMode	? CreateFlagLoosener<uint8 >(Schema, bNeeded, Ids, Scratch)
+														: CreateFlatLoosener<uint8 >(Schema, bNeeded, Ids, Scratch);
+		case ELeafWidth::B16:	return Schema.FlagMode	? CreateFlagLoosener<uint16>(Schema, bNeeded, Ids, Scratch)
+														: CreateFlatLoosener<uint16>(Schema, bNeeded, Ids, Scratch);
+		case ELeafWidth::B32:	return Schema.FlagMode	? CreateFlagLoosener<uint32>(Schema, bNeeded, Ids, Scratch)
+														: CreateFlatLoosener<uint32>(Schema, bNeeded, Ids, Scratch);
+		case ELeafWidth::B64:	return Schema.FlagMode	? CreateFlagLoosener<uint64>(Schema, bNeeded, Ids, Scratch)
+														: CreateFlatLoosener<uint64>(Schema, bNeeded, Ids, Scratch);
+		default:				check(false); return FEnumLoosener();
+	}
 }
 
 // Contiguous constants will match contiguous declaration order of FEnumeratorId indices
@@ -1871,6 +2497,21 @@ static FEnumFastener CreateEnumFastener(const FEnumDeclaration* Decl, FScratchAl
 	return { Decl && RequiresFastening(Decl->GetEnumerators(), Decl->Mode) ? Decl : nullptr };
 }
 
+static const FLooseSchema& ObtainSourceSchema(FBatchContext& Ctx, FDeclId Id, FStructSchemaId SchemaId)
+{
+	const FTransformation& Xform = Ctx.Transformations[Id];
+	check(!Xform.Upgrades);
+	const FLooseSchema*& Cache = Xform.SourceSchema;
+	if (!Cache)
+	{
+		const FSourceSchemas& Schemas = Ctx.Transformations.Schemas;
+		const FStructSchema& Schema = GetStructSchemas(*Schemas.Batch)[SchemaId.Idx];
+		Cache = CreateSourceSchema(Schema, Schemas, Ctx.LooseIds.Get(Id), Ctx.Scratch);
+	}
+
+	return *Cache;
+}
+
 void FBatchRun::CreateTransformations(const TBitArray<>& MatchingSchemas, const FSchemaBatch& Batch)
 {
 	FLooseIndexers& LooseIds = History.LooseIds;
@@ -1878,21 +2519,16 @@ void FBatchRun::CreateTransformations(const TBitArray<>& MatchingSchemas, const
 	const TMap<FDeclId, FStructUpgradesPtr>& CompiledUpgrades = History.CompiledUpgrades;
 	const TSchemaRange<const FStructSchema> StructSchemas = GetStructSchemas(Batch);
 	const TSchemaRange<const FEnumSchema> EnumSchemas = GetEnumSchemas(Batch);
-	const TArray<uint16> StructVersions = ExtractVersions(StructSchemas);
-	const TBitArray<> EnumFlagModes = ExtractFlagModes(EnumSchemas);
-	const uint32 NumSourceStructs = StructVersions.Num();
+	Transformations.Schemas = { uint32(StructSchemas.Num()), &Batch, RuntimeIds, ExtractVersions(StructSchemas),  ExtractFlagModes(EnumSchemas) }; 
 
 	TArray<FTransformation>& StructXforms = Transformations.Structs;
 	TSet<FDeclId>& StructLookup = Transformations.StructLookup;
 	TArray<FEnumTransformation>& EnumXforms = Transformations.Enums;
 	TSet<FEnumId>& EnumLookup = Transformations.EnumLookup;
-	Transformations.NumSourceStructs = StructSchemas.Num();
 	StructXforms.SetNum(StructSchemas.Num());
 	StructLookup.Reserve(StructSchemas.Num());
 	EnumXforms.SetNum(EnumSchemas.Num());
 	EnumXforms.Reserve(EnumSchemas.Num());
-	
-	TSet<FEnumId> EnumDecompositions;
 
 	// Create initial struct transforms
 	for (uint32 Idx = 0; int32(Idx) < StructSchemas.Num(); ++Idx)
@@ -1902,7 +2538,7 @@ void FBatchRun::CreateTransformations(const TBitArray<>& MatchingSchemas, const
 		{
 			FParameterIndexer& Ids = LooseIds.Get(Id);
 			FTransformation& Xform = StructXforms[Idx];
-			Xform.SourceSchema = CreateSourceSchema(StructSchemas[Idx], Ids, StructVersions, EnumFlagModes);
+			Xform.SourceSchema = CreateSourceSchema(StructSchemas[Idx], Transformations.Schemas, Ids, Scratch);
 			Xform.Upgrades = CompiledUpgrades.FindChecked(Id).Get();
 			Xform.TargetSchema = History.LatestSchemas.Obtain(Id, Declarations, Ids);
 		}
@@ -1930,6 +2566,7 @@ void FBatchRun::CreateTransformations(const TBitArray<>& MatchingSchemas, const
 	
 	// Find renames and inner types outputted by upgrade ops,
 	// which aren't part of schema batch
+	TBitArray<> EnumInputSchemas(false, EnumSchemas.Num());
 	TArrayView<FTransformation> Unscanned = StructXforms;
 	while (Unscanned.Num() > 0)
 	{
@@ -1942,6 +2579,7 @@ void FBatchRun::CreateTransformations(const TBitArray<>& MatchingSchemas, const
 				// for (const FVersionUpgrades* Version : ViewVersionsFrom(*Upgrades, FirstVersion))
 				for (const FVersionUpgrades* Version : ViewVersions(*Upgrades))
 				{
+					AddEnumInputSchemas(/* out */ EnumInputSchemas, EnumLookup, *Version);
 					AddInnerIdOutputs(/* out */ StructLookup, /* out */ EnumLookup, *Version);
 				}
 
@@ -1967,18 +2605,22 @@ void FBatchRun::CreateTransformations(const TBitArray<>& MatchingSchemas, const
 		}
 	}
 
-	// Create source looseners for enum input decompositions (that might not be needed)
-	for (FEnumId Id : EnumDecompositions)
+	// Create source looseners for all input enums
+	TBitArray<> MissingLooseners;
+	MissingLooseners.AddRange(MatchingSchemas, EnumSchemas.Num(), StructSchemas.Num());
+	MissingLooseners.CombineWithBitwiseAND(EnumInputSchemas, EBitwiseOperatorFlags::MinSize);
+	for (TConstSetBitIterator<> It(MissingLooseners); It; ++It)
 	{
-		int32 Idx = EnumLookup.FindId(Id).AsInteger();
-		if (Idx > INDEX_NONE && Idx < EnumXforms.Num())
-		{
-			FEnumLoosener& Source = EnumXforms[Idx].Source;
-			Source = Source.Exists() ? Source : CreateEnumLoosener(EnumSchemas[Idx], /* bNeeded */ false, LooseIds.Get(Id), Scratch);
-		}
+		int32 Idx = It.GetIndex();
+		FEnumId Id = EnumLookup.Get(FSetElementId::FromInteger(Idx));
+		FEnumTransformation& Xform = EnumXforms[Idx];
+		check(!Xform.Source.Exists());
+		Xform.Source = CreateEnumLoosener(EnumSchemas[Idx], /* bNeeded */ false, LooseIds.Get(Id), Scratch);
+		// Create fastener too, only needed for output enums
+		Xform.Loose = CreateEnumFastener(Declarations.Find(Id), Scratch);
 	}
 
-	// Create fasteners for enum output recompositions
+	// Create fasteners for enum outputs that lack schema in the batch
 	EnumXforms.SetNum(EnumLookup.Num());
 	TArrayView<FEnumTransformation> OutputEnumXforms = MakeArrayView(EnumXforms).RightChop(EnumSchemas.Num());
 	for (int32 Idx = EnumSchemas.Num(); FEnumTransformation& Xform : OutputEnumXforms)
@@ -1998,7 +2640,7 @@ static FLooseType IndexObsolete(FLiteralIndexerBase& Ids)
 }
 
 // @param In isn't mismatched
-static FLooseStruct ReorderAsDeclaration(FLooseStruct In, const FLooseSchema& Declared, FScratchAllocator& Scratch)
+static FLooseStruct ReorderAsDeclaration(FLooseStruct In, const FLooseSchema& Declared, EMismatch& OutMismatch, FScratchAllocator& Scratch, FDebugPrinter Debug)
 {
 	// opt: check if In members have consecutive names, if so type-check and return In
 	//		could also compact and only return IsSet members and reuse the existing In allocation
@@ -2009,26 +2651,45 @@ static FLooseStruct ReorderAsDeclaration(FLooseStruct In, const FLooseSchema& De
 		new (OutIt++) FLooseMember(MakeMember(Param));
 	}
 
+	TArray<FLooseMember> Mismatches;
 	TArrayView<FLooseMember> Out(OutIt - Declared.NumMembers, Declared.NumMembers);
 	for (FLooseMember Src : MakeArrayView(In.Members, In.Num))
 	{
 		if (Src)
 		{
-			FParameterId Name = GetName(Src);
-			check(Name.Idx < Out.Num()); // Should've been a mismatch
-			FLooseMember& Dst = Out[Name.Idx];
-			check(!Dst);
-			check(Src.GetType().Compare(Dst.GetType()));
-			Dst.SetTypechecked(Src.GetValue());
+			FParameterId Name = Src.GetName();
+			FLooseMember* Dst = Name.Idx < Out.Num() ? &Out[Name.Idx] : nullptr;
+			if (Dst && Src.GetType().Compare(Dst->GetType()) && !*Dst)
+			{
+				Dst->SetTypechecked(Src.GetValue());
+			}
+			else
+			{
+				Src.EditTypechecked().SetFirstMismatch(Dst ? (*Dst ? EMismatch::Duplicate : EMismatch::Type): EMismatch::Surplus);
+				Mismatches.Add(Src);
+			}
 		}
 	}
 
+	if (Mismatches.Num())
+	{
+		// Report first mismatch type
+		OutMismatch = Mismatches[0].GetValue().Meta.Mismatch;
+		check(OutMismatch != EMismatch::No);
+
+		int32 NumAll = Declared.NumMembers + Mismatches.Num();
+		FLooseMember* All = static_cast<FLooseMember*>(Scratch.AllocateZeroed(NumAll * sizeof(FLooseMember), alignof(FLooseMember)));
+		FMemory::Memcpy(All, Out.GetData(), Out.NumBytes());
+		FMemory::Memcpy(All + Declared.NumMembers, Mismatches.GetData(), Mismatches.NumBytes());
+		Out = MakeArrayView(All, NumAll);
+	}
+
 	return { In.Id, In.Version, static_cast<uint16>(Out.Num()), Out.GetData() };
 }
 
 FUpgradedInstance FBatchRun::UpgradeStruct(FStructView Source, FTransformation Xform) const
 {
-	FBatchContext Ctx = { Transformations, Scratch, History.Ids, History.LooseIds, IndexObsolete(History.Ids), Source.Schema.Batch, {Scratch}, {}, Debug };
+	FBatchContext Ctx = { Transformations, Scratch, History.Ids, History.LooseIds, IndexObsolete(History.Ids), Source.Schema.Batch, {Scratch}, FContext(Scratch), Debug };
 	FInstanceUpgrader SourceUpgrader(Ctx, *Xform.Upgrades, Xform.SourceSchema->Version);
 	EMismatch Mismatch = EMismatch::No;
 	FLooseStruct Upgraded = SourceUpgrader.Upgrade(Source, *Xform.SourceSchema, Xform.TargetSchema, /* out */ Mismatch);
@@ -2044,7 +2705,8 @@ FUpgradedInstance FBatchRun::UpgradeStruct(FStructView Source, FTransformation X
 		}
 		else
 		{
-			Upgraded = ReorderAsDeclaration(Upgraded, *Xform.TargetSchema, Ctx.Scratch);
+			Upgraded = ReorderAsDeclaration(Upgraded, *Xform.TargetSchema, /* out */ Mismatch, Scratch, Ctx.Debug);
+			break;
 		}
 	}
 
@@ -2136,7 +2798,7 @@ FBuiltStruct* FBatchRun::FastenStruct(FLooseStruct Struct) const
 void FBatchRun::FastenMembers(FMemberBuilder& Out, TConstArrayView<FLooseMember> Members, const FParameterIndexer& Names, const FStructDeclaration& Declaration) const
 {
 	// Fasten super
-	if (Names.Resolve(GetName(Members[0])) == NoId)
+	if (Names.Resolve(Members[0].GetName()) == NoId)
 	{
 		FStructMember Super = AsStruct(Members[0]);
 		Members.RightChopInline(1);
@@ -2154,7 +2816,7 @@ void FBatchRun::FastenMembers(FMemberBuilder& Out, TConstArrayView<FLooseMember>
 			// todo: decide pass down Declared.GetInnermostType() to get enum width or not
 			// maybe pass down innermost leaf size, 0 if not leaf
 			FBuiltValue Value = FastenMember(Member /*, Declared.GetInnermostType()*/);
-			FMemberId Name = Names.Resolve(GetName(Member)).Get();
+			FMemberId Name = Names.Resolve(Member.GetName()).Get();
 			Out.Add(Name, {Declared, Value});
 		}
 	}
@@ -2196,7 +2858,8 @@ FBuiltStruct* FBatchRun::FastenStruct(FStructMember Struct) const
 		FTransformation Xform = Transformations[Struct.Id];
 		checkf(!Xform.Upgrades, TEXT("Struct should've been upgraded already"));
 		check(Xform.TargetSchema);
-		Loose = ReorderAsDeclaration(Loose, *Xform.TargetSchema, Scratch);
+		Loose = ReorderAsDeclaration(Loose, *Xform.TargetSchema, /* out */ Struct.Mismatch, Scratch, Debug);
+		check(Struct.Mismatch == EMismatch::No);
 	}
 		
 	return FastenStruct(Loose);
@@ -2211,8 +2874,8 @@ uint64 FBatchRun::FastenEnum(FEnumMember Enum) const
 	//checkf(Enum.State == ELooseState::UpgradedLoose, TEXT("Enum not fully upgraded"));
 	FEnumFastener Fastener = Transformations[Enum.Id].Loose;
 	const FEnumeratorIndexer& Indexer = History.LooseIds.Get(Enum.Id); // opt: No need to lookup if no fastener...
-	return Enum.FlagMode ? Fasten(Fastener, Enum.Data.Loose.Enumerators.Flags, Indexer)
-						 : Fasten(Fastener, Enum.Data.Loose.Enumerators.Flat, Indexer);
+	return Enum.FlagMode ? Fasten(Fastener, Enum.Data.Loose.Enumerator.Flags, Indexer)
+						 : Fasten(Fastener, Enum.Data.Loose.Enumerator.Flat, Indexer);
 }
 
 FBuiltRange* FBatchRun::FastenRange(FRangeMember Range) const
@@ -2226,17 +2889,19 @@ FBuiltRange* FBatchRun::FastenRange(FLooseType Type, FLooseMetadata Meta, FRange
 	{
 		return FastenRanges(Type, Meta, MakeArrayView(Data.Ranges, Meta.Num));
 	}
-	else if (Type.IsStruct())
+	else if (Type.InnermostIsStruct())
 	{
+		check(!Meta.IsSourceRange); // todo: call RebuildRanges. Should only matter for enums/structs
 		return FastenStructs(Type.InnermostId.Get().AsStructDeclId(), Meta, Data.Structs);
 	}
-	else if (Type.IsEnum())
+	else if (Type.InnermostIsEnum())
 	{
+		check(!Meta.IsSourceRange); // todo: call RebuildRanges. Should only matter for enums/structs
 		return FastenEnums(Type.InnermostId.Get().AsEnum(), Type.InnermostType.Enum, Meta, Data);
 	}
 	else
 	{
-		return FastenArithmetics(Type.InnermostType.Kind, Meta.Num, Data.Leaves);
+		return FastenArithmetics(Type.InnermostType.Kind, Meta.Num, Data.Arithmetics.Source);
 	}
 }
 
@@ -2250,29 +2915,27 @@ FBuiltRange* FBatchRun::FastenRanges(FLooseType Type, FLooseMetadata Meta, TCons
 	FBuiltRange* Out = FBuiltRange::Create(Scratch, Ranges.Num(), sizeof(FBuiltRange*));
 	FBuiltRange** OutIt = reinterpret_cast<FBuiltRange**>(Out->Data);
 
-	FLooseType InnerType = Type;
-	InnerType.NumRanges--;
-
-	if (InnerType.IsRange())
+	FLooseType ItemType = ItemOf(Type);
+	if (ItemType.IsRange())
 	{
 		for (FLooseNestedRange Range : Ranges)
 		{
-			*OutIt++ = FastenRanges(InnerType, Meta, MakeArrayView(Range.Data.Ranges, Range.Num));
+			*OutIt++ = FastenRanges(ItemType, Meta, MakeArrayView(Range.Data.Ranges, Range.Num));
 		}
 	}
-	else if (InnerType.IsStruct())
+	else if (ItemType.IsStruct())
 	{
-		FDeclId Id = InnerType.InnermostId.Get().AsStructDeclId();
+		FDeclId Id = ItemType.InnermostId.Get().AsStructDeclId();
 		for (FLooseNestedRange Range : Ranges)
 		{
 			Meta.Num = Range.Num;
 			*OutIt++ = FastenStructs(Id, Meta, Range.Data.Structs);
 		}
 	}
-	else if (InnerType.IsEnum())
+	else if (ItemType.IsEnum())
 	{
-		FEnumId Id = InnerType.InnermostId.Get().AsEnum();
-		FInnermostEnumType InnermostType = InnerType.InnermostType.Enum;
+		FEnumId Id = ItemType.InnermostId.Get().AsEnum();
+		FInnermostEnumType InnermostType = ItemType.InnermostType.Enum;
 		for (FLooseNestedRange Range : Ranges)
 		{
 			Meta.Num = Range.Num;
@@ -2283,7 +2946,7 @@ FBuiltRange* FBatchRun::FastenRanges(FLooseType Type, FLooseMetadata Meta, TCons
 	{
 		for (FLooseNestedRange Range : Ranges)
 		{
-			*OutIt++ = FastenArithmetics(InnerType.InnermostType.Kind, Range.Num, Range.Data.Leaves);
+			*OutIt++ = FastenArithmetics(ItemType.InnermostType.Kind, Range.Num, Range.Data.Arithmetics.Source);
 		}
 	}
 
@@ -2296,23 +2959,69 @@ FBuiltRange* FBatchRun::FastenStructs(FDeclId Id, FLooseMetadata Meta, FStructRa
 	{
 		return nullptr;
 	}
-	else if (Meta.InnermostSchema)
-	{
-		FStructSchemaHandle Schema = { FStructSchemaId(Meta.InnermostSchema.Get()), SourceBatch };
-		return RebuildStructs(FStructRangeView(Meta.Num, ViewSkippableSlice(Structs.Source), Schema), Id);
-	}
-
+	
 	const FParameterIndexer& Names = History.LooseIds.Get(Id);
 	const FStructDeclaration& Declaration = History.Latest.Get(Id);
 
 	FBuiltRange* Out = FBuiltRange::Create(Scratch, Meta.Num, sizeof(FBuiltStruct*));
 	FBuiltStruct** OutIt = reinterpret_cast<FBuiltStruct**>(Out->Data);
-
 	FMemberBuilder Tmp;
-	for (FLooseValue Struct : TConstArrayView64<FLooseValue>(Structs.Loose, Meta.Num))
+	if (Meta.InnermostSchema)
+	{
+		FStructSchemaHandle Schema = { FStructSchemaId(Meta.InnermostSchema.Get()), SourceBatch };
+		for (FSourceStructRangeReader SourceIt(Structs, Meta.Num); SourceIt.HasMore();)
+		{
+			RebuildMembers(Tmp, SourceIt.GrabView(Schema), Id);
+			*OutIt++ = Tmp.BuildAndReset(Scratch, Declaration, Debug); 
+		}
+	}
+	else
+	{
+		for (FLooseValue Struct : TConstArrayView64<FLooseValue>(Structs.Loose, Meta.Num))
+		{
+			if (Struct.Meta.InnermostSchema)
+			{
+				RebuildMembers(Tmp, AsSourceStruct(Struct, SourceBatch), Id);
+			}
+			else
+			{
+				FastenMembers(Tmp, MakeArrayView(Struct.Data.Struct.Loose, Struct.Meta.Num), Names, Declaration);
+			}
+			*OutIt++ = Tmp.BuildAndReset(Scratch, Declaration, Debug); 
+		}
+		
+	}
+	check(OutIt == Out->AsStructs().GetData() + Meta.Num);
+	return Out;
+}
+
+template<typename T>
+FBuiltRange* FastenLooseEnums(TConstArrayView64<FEnumLooseData> Ins, FEnumFastener Fastener, const FEnumeratorIndexer& Indexer, FScratchAllocator& Scratch)
+{
+	const uint64 Num = Ins.Num();
+	FBuiltRange* Out = FBuiltRange::Create(Scratch, Num, sizeof(T)); 
+	T* OutIt = reinterpret_cast<T*>(Out->Data);
+	
+	if (Fastener.Decl == nullptr)
+	{
+		for (FEnumLooseData In : Ins)
+		{
+			*OutIt++ = IntCastChecked<T>(BitCast<uint64>(In.Enumerator));
+		}
+	}
+	else if (Fastener.Decl->Mode == EEnumMode::Flat)
 	{
-		FastenMembers(Tmp, MakeArrayView(Struct.Data.Struct.Loose, Struct.Meta.Num), Names, Declaration);
-		*OutIt++ = Tmp.BuildAndReset(Scratch, Declaration, Debug); 
+		for (FEnumLooseData In : Ins)
+		{
+			*OutIt++ = IntCastChecked<T>(Fasten(Fastener, In.Enumerator.Flat, Indexer));
+		}
+	}
+	else
+	{
+		for (FEnumLooseData In : Ins)
+		{
+			*OutIt++ = IntCastChecked<T>(Fasten(Fastener, In.Enumerator.Flags, Indexer));
+		}
 	}
 
 	return Out;
@@ -2326,13 +3035,23 @@ FBuiltRange* FBatchRun::FastenEnums(FEnumId Id, FInnermostEnumType Type, FLooseM
 	}
 	else if (Meta.InnermostSchema)
 	{
-		SIZE_T EnumSize = SizeOf(History.Latest.Get(Id).Width.Get());
-		return FBuiltRange::CreateIf(Scratch, Meta.Num, EnumSize, Data.Leaves); 
+		const FEnumSchema& Schema = ResolveEnumSchema(SourceBatch, FEnumSchemaId(Meta.InnermostSchema.Get()));
+		return FBuiltRange::Create(Scratch, Meta.Num, SizeOf(Schema.Width), Data.Enums.Source); 
 	}
 
-	// todo: enum range fastening
-	unimplemented();
+	TArrayView64<FEnumLooseData> Enums(Data.Enums.Loose, Meta.Num);
+	FEnumFastener Fastener = Transformations[Id].Loose;
+	const FEnumeratorIndexer& Indexer = History.LooseIds.Get(Id);
+	const FEnumDeclaration& Declaration = History.Latest.Get(Id);
+	switch (Declaration.Width.Get())
+	{
+		case ELeafWidth::B8:		return FastenLooseEnums<uint8 >(Enums, Fastener, Indexer, Scratch);
+		case ELeafWidth::B16:		return FastenLooseEnums<uint16>(Enums, Fastener, Indexer, Scratch);
+		case ELeafWidth::B32:		return FastenLooseEnums<uint32>(Enums, Fastener, Indexer, Scratch);
+		case ELeafWidth::B64:		return FastenLooseEnums<uint64>(Enums, Fastener, Indexer, Scratch);
+	}
 
+	check(false);
 	return nullptr;
 }
 
@@ -2713,7 +3432,7 @@ struct FOps final : IOps
 		return InnerMembers.AddSlice(Slice);
 	}
 
-	virtual FEnumId	NameEnum(FTypename Enum) override
+	virtual FEnumId NameEnum(FTypename Enum) override
 	{
 		return Ids.IndexEnum(IndexType(Ids, Enum));
 	}
@@ -2931,7 +3650,7 @@ static void CompileDecomposition(uint8*& InOutIt, FMidParam Input, TInnerSlices<
 {
 	FInnerId Id = Input.Type.InnermostId.Get();
 	uint16 Version = Input.Type.InnermostVersion;
-	if (Input.Type.IsStruct())
+	if (Input.Type.InnermostIsStruct())
 	{
 		TConstArrayView<FMidParam> Inners = Members.GetSlice(Input.InnerSlice);
 		new (InOutIt) FStructDecomposition{ Id.AsStructDeclId(), Version,  CreateInputOrder(Inners), /* set in Link */ nullptr, 
@@ -2940,7 +3659,7 @@ static void CompileDecomposition(uint8*& InOutIt, FMidParam Input, TInnerSlices<
 	}
 	else
 	{
-		check(Input.Type.IsEnum());
+		check(Input.Type.InnermostIsEnum());
 		TConstArrayView<FEnumeratorId> Inners = Enumerators.GetSlice(Input.InnerSlice);
 		new (InOutIt) FEnumDecomposition{ CreateInputOrder(Inners) };
 		InOutIt += sizeof(FEnumDecomposition);
@@ -2992,9 +3711,10 @@ static void CompileOps(FVersionUpgrades& Out, const FOps& In)
 			// Generate input decomposition match
 			if (Input.InnerSlice)
 			{
-				FMatchIdx Idx = { .Decomp = DecompIdx };	
-				AllMatches[Input.Name.Idx].Add({ .Type = Input.Type, .Idx = Idx, .InputIdx = InputIdx, .Recipient = FMemberMatch::DecomposeRecipient });
-				
+				FMatchIdx Idx = { .Decomp = DecompIdx };
+				FLooseType Type = ClearDecomposed(Input.Type);
+				AllMatches[Input.Name.Idx].Add({ .Type = Type, .Idx = Idx, .InputIdx = InputIdx, .Recipient = FMemberMatch::DecomposeRecipient });
+
 				// Write offset and generate decomposition
 				DecompOffsets[DecompIdx] = IntCastChecked<uint32>(DecompDataIt - reinterpret_cast<uint8*>(&Out));
 				CompileDecomposition(/* in-out */ DecompDataIt, Input, In.InnerMembers, In.InnerEnumerators);
diff --git a/Engine/Plugins/Experimental/PlainProps/Source/Public/PlainPropsBuild.h b/Engine/Plugins/Experimental/PlainProps/Source/Public/PlainPropsBuild.h
--- a/Engine/Plugins/Experimental/PlainProps/Source/Public/PlainPropsBuild.h
+++ b/Engine/Plugins/Experimental/PlainProps/Source/Public/PlainPropsBuild.h
@@ -72,6 +72,17 @@ public:
 		}
 		return Out;
 	}
+	
+	template<typename T>
+	inline T* CloneArray(uint64 Num, const T* In)
+	{
+		T* Out = static_cast<T*>(Allocate(Num * sizeof(T), alignof(T)));
+		for (uint64 Idx = 0; Idx < Num; ++Idx)
+		{
+			new (Out + Idx) T(In[Idx]);
+		}
+		return Out;
+	}
 
 	PLAINPROPS_API void Reset();
 };
diff --git a/Engine/Plugins/Experimental/PlainProps/Source/Public/PlainPropsLooseTypes.h b/Engine/Plugins/Experimental/PlainProps/Source/Public/PlainPropsLooseTypes.h
--- a/Engine/Plugins/Experimental/PlainProps/Source/Public/PlainPropsLooseTypes.h
+++ b/Engine/Plugins/Experimental/PlainProps/Source/Public/PlainPropsLooseTypes.h
@@ -18,25 +18,29 @@ enum class EInnermostKind : uint8 { Struct, Enum, Bool, S8, S16, S32, S64, U8, U
 
 struct FInnermostStructType
 {
-	EInnermostKind			_ : 6;
+	EInnermostKind			_ : 5;
 	uint8					IsDynamic : 1;
-	uint8					IsSuper : 1; // Needed given FOptionalMemberId-indexing?
+	uint8					IsSuper : 1;
+	uint8					IsDecomposed : 1;
 };
 
 struct FInnermostEnumType
 {
-	EInnermostKind			_ : 6;
+	EInnermostKind			_ : 5;
 	uint8					FlagMode : 1;
 	uint8					Unused : 1;
+	uint8					IsDecomposed : 1;
 };
 
 union FInnermostType
 {
-	EInnermostKind			Kind : 6;
+	EInnermostKind			Kind : 5;
 	FInnermostStructType	Struct;
 	FInnermostEnumType		Enum;
 };
 
+inline const bool operator==(FInnermostType A, FInnermostType B) { return BitCast<uint8>(A) == BitCast<uint8>(B); }
+
 inline static constexpr FInnermostType InnermostFlatEnum = { .Enum = { ._ = EInnermostKind::Enum, .FlagMode = 0 } };
 inline static constexpr FInnermostType InnermostFlagEnum = { .Enum = { ._ = EInnermostKind::Enum, .FlagMode = 1 } };
 inline static constexpr FInnermostType InnermostStaticStruct = { .Struct = { ._ = EInnermostKind::Struct, .IsDynamic = 0 } };
@@ -101,6 +105,8 @@ struct FLooseType
 	bool				IsStruct() const				{ return NumRanges == 0 && InnermostType.Kind == EInnermostKind::Struct; }
 	bool				IsLeaf() const					{ return NumRanges == 0 && InnermostType.Kind != EInnermostKind::Struct; }
 	bool				IsEnum() const					{ return NumRanges == 0 && InnermostType.Kind == EInnermostKind::Enum; }
+	bool				InnermostIsStruct() const		{ return InnermostType.Kind == EInnermostKind::Struct; } 
+	bool				InnermostIsEnum() const			{ return InnermostType.Kind == EInnermostKind::Enum; } 
 
 	bool				Compare(FLooseType O) const		{ return WithoutMetadata().AsInt() == O.WithoutMetadata().AsInt(); }
 	bool				CompareAll(FLooseType O) const	{ return AsInt() == O.AsInt(); }
@@ -140,64 +146,85 @@ struct FEnumeratorIdSet
 
 union FEnumeratorData
 {
-	FEnumeratorId		Flat;	// Loose enumerator name (without integer value)
-	FEnumeratorIdSet	Flags;	// Loose set of enumerator flag names without integer value
+	FEnumeratorId		Flat = {};		// Loose enumerator name (without integer value)
+	FEnumeratorIdSet	Flags;			// Loose set of enumerator flag names without integer value
 };
 
-struct FEnumeratorInput
+// Temporary format used to access individual named enumerators during Op execution
+struct FEnumeratorIndex
 {
-	uint16				Index;
+	uint16				Index = 0xFFFF;
 };
 
-struct FEnumeratorInputSet
+// Temporary format used to access individual named enum flags during Op execution
+struct FEnumeratorIndexSet
 {
-	uint64				Bits;
+	uint64				Bits = 0;
 };
 
 // Temporary format used to access individual enum members during Op execution
 union FEnumIndexData
 {
-	FEnumeratorInput	Flat;	// Enumerator index using NameEnumerators order
-	FEnumeratorInputSet	Flags;	// Bitset of flag indices using NameFlags order
+	FEnumeratorIndex	Flat = {};			// Enumerator index from NameEnumerators order
+	FEnumeratorIndexSet	Flags;				// Bitset of flag indices from NameFlags order
 };
 
 union FEnumLooseData
 {
-	FEnumeratorData		Enumerators;	// FEnumeratorId::Idx order from FEnumeratorIndexer
-	FEnumIndexData		Indices;		// Input order from NameEnumerators/NameFlags
+	FEnumeratorData		Enumerator = {};	// FEnumeratorId::Idx order from FEnumeratorIndexer
+	FEnumIndexData		Decomposed;			// Index order from NameEnumerators/NameFlags
 };
 
 union FEnumData
 {
-	uint64				Source = 0;	// Integer source value matching the innermost enum schema
+	uint64				Source = 0;			// Integer source value matching the innermost enum schema
 	FEnumLooseData		Loose;
 };
 
+//////////////////////////////////////////////////////////////////////////////////////////////////
+
+union FEnumRange
+{
+	const void*				Source = nullptr;		// Integers source values matching the innermost enum schema
+	FEnumLooseData*			Loose;				
+};
+
 union FStructRange
 {
-	const void*			Source = nullptr; // Skippable slice
-	const FLooseValue*	Loose;
+	const void*				Source = nullptr;		// Skippable slice of the FIRST source FStructView, not the range of structs
+	FLooseValue*			Loose;
+};
+
+union FArithmeticRange
+{
+	const void*				Source = nullptr;
+	void*					Loose;
 };
 
 union FRangeData
 {
-	const void*					Leaves = nullptr;
-	FStructRange				Structs;
-	const FLooseNestedRange*	Ranges; 
+	FArithmeticRange		Arithmetics = {};
+	FEnumRange				Enums;
+	FStructRange			Structs;
+	FLooseNestedRange*		Ranges; 
 };
 
 struct FLooseNestedRange
 {
-	uint64				Num = 0;
-	FRangeData			Data;
+	uint64					Num = 0;
+	FRangeData				Data;
 };
 
+//////////////////////////////////////////////////////////////////////////////////////////////////
+
 union FStructData
 {
 	const uint8*		Source = nullptr;	// Skippable slice from a source FStructView
 	FLooseMember*		Loose;
 };
 
+//////////////////////////////////////////////////////////////////////////////////////////////////
+
 union FLooseData
 {
 	uint64				Arithmetic = 0;
@@ -206,9 +233,10 @@ union FLooseData
 	FStructData			Struct;
 };
 
+//////////////////////////////////////////////////////////////////////////////////////////////////
+
 enum class EMismatch : uint8 { No, Inners, Inputs, Type, Duplicate, Surplus, Missing };
 
-// FLooseData metadata
 struct FLooseMetadata
 {
 	uint8				IsSet = 0;
@@ -247,7 +275,7 @@ inline void FLooseValue::SetFirstMismatch(EMismatch Mismatch)
 
 //////////////////////////////////////////////////////////////////////////////////////////////////
 
-// Type and named optional value
+// Typed and named optional value
 class FLooseMember
 {
 public:
@@ -258,9 +286,13 @@ public:
 	bool				IsOptional() const				{ return Type.OptionalParameter; }
 	FLooseType			GetType() const					{ return Type; }
 	FLooseValue			GetValue() const				{ return Value; }
+	FParameterId		GetName() const					{ return Value.Meta.Name; }
 	void				Rename(FParameterId Name)		{ Value.Meta.Name = Name; }
 	void				SetTypechecked(FLooseValue V)	{ Value = V; } 
-	FLooseValue&		EditTypechecked()				{ return Value; } 
+	FLooseValue&		EditTypechecked()				{ return Value; }
+	void				RecordDynamicType(FLooseType T);
+	void				MarkEnumRecomposed()			{ Type.InnermostType.Enum.IsDecomposed = 0; }
+	void				MarkStructRecomposed()			{ Type.InnermostType.Struct.IsDecomposed = 0; }
 
 protected:
 	FLooseType			Type;
diff --git a/Engine/Plugins/Experimental/PlainProps/Source/Public/PlainPropsSpecify.h b/Engine/Plugins/Experimental/PlainProps/Source/Public/PlainPropsSpecify.h
--- a/Engine/Plugins/Experimental/PlainProps/Source/Public/PlainPropsSpecify.h
+++ b/Engine/Plugins/Experimental/PlainProps/Source/Public/PlainPropsSpecify.h
@@ -13,7 +13,7 @@ enum ESpecBool { SpecBool };
 enum ESpecIntS { SpecS8, SpecS16, SpecS32, SpecS64 };			// Values match ELeafWidth
 enum ESpecIntU { SpecU8, SpecU16, SpecU32, SpecU64 };			// Values match ELeafWidth
 enum ESpecHex { SpecHex8, SpecHex16, SpecHex32, SpecHex64 };	// Values match ELeafWidth
-enum ESpecUtf { SpecUtf8, SpecUtf16, SpecUtf32 };				// Values match ELeafWidth
+enum ESpecChar { SpecUtf8, SpecUtf16, SpecUtf32 };				// Values match ELeafWidth
 enum ESpecFloat { SpecF32 = 2, SpecF64 = 3 };					// Values match ELeafWidth
 enum ESpecDynamicStruct { SpecDynamicStruct };
 
@@ -39,7 +39,7 @@ public:
 	FMemberSpec(ESpecIntS In)	: FMemberSpec(ELeafType::IntS,		static_cast<ELeafWidth>(In), NoId) {}
 	FMemberSpec(ESpecIntU In)	: FMemberSpec(ELeafType::IntU,		static_cast<ELeafWidth>(In), NoId) {}
 	FMemberSpec(ESpecHex In)	: FMemberSpec(ELeafType::Hex,		static_cast<ELeafWidth>(In), NoId) {}
-	FMemberSpec(ESpecUtf In)	: FMemberSpec(ELeafType::Unicode,	static_cast<ELeafWidth>(In), NoId) {}
+	FMemberSpec(ESpecChar In)	: FMemberSpec(ELeafType::Unicode,	static_cast<ELeafWidth>(In), NoId) {}
 	FMemberSpec(ESpecFloat In)	: FMemberSpec(ELeafType::Float,		static_cast<ELeafWidth>(In), NoId) {}
 	
 	FMemberSpec(FDeclId Id)			: FMemberSpec(FStructType{EMemberKind::Struct, /* IsDynamic */ 0, /* IsSuper */ 0}, FInnerId(Id)) {}
diff --git a/Engine/Plugins/Experimental/PlainProps/Source/Public/PlainPropsUpgrade.h b/Engine/Plugins/Experimental/PlainProps/Source/Public/PlainPropsUpgrade.h
--- a/Engine/Plugins/Experimental/PlainProps/Source/Public/PlainPropsUpgrade.h
+++ b/Engine/Plugins/Experimental/PlainProps/Source/Public/PlainPropsUpgrade.h
@@ -3,6 +3,7 @@
 #pragma once
 
 #include "PlainPropsTypes.h"
+#include "PlainPropsRead.h" // FBoolRangeIterator
 #include "PlainPropsLooseTypes.h"
 #include "Containers/BitArray.h"
 #include "Misc/Optional.h"
@@ -30,21 +31,32 @@ namespace Upgrade
 class FChronicler;
 class FContext;	
 class FEnumChronicler;
+struct FFriend;
 struct FHistory;
 class FIn;
+struct FInAny;
 class FInEnum;
 class FInEnumerator;
-class FInEnumFlags;
-class FInStruct;
+class FInFlags;
+struct FInMembers;
+struct FInStruct;
 struct IOps;
 class FOut;
-class FOutEnum;
-class FOutEnumerator;
-class FOutEnumFlags;
+class FOutAny;
+class FOutBoolItem;
+struct FOutEnum;
+class FOutEnumItem;
+struct FOutEnumerator;
+class FOutEnumeratorItem;
+struct FOutFlags;
+class FOutFlagsItem;
+struct FOutMembers;
 struct FOutParameter;
-class FOutStruct;
+struct FOutStruct;
 class FUpgrades;
-template<uint32 N> class TOuts;
+template<typename T, uint8 N> struct TInIt;
+template<typename T, uint8 N> struct TOutIt;
+template<typename T, uint8 N = 1> struct TOutRange;
 
 using FIns = TConstArrayView<FIn>;
 using FOuts = TArrayView<FOut>;
@@ -80,8 +92,9 @@ struct FParameter
 	bool				bOptional;
 	FInnermostType		InnermostType;
 	FInnermostHandle	InnermostHandle;
-
-	FParameter operator~() const { return {Name, NameLen, NumRanges, true, InnermostType, InnermostHandle}; }
+	
+	constexpr FParameter operator~() 			{ return {Name, NameLen, NumRanges, true, InnermostType, InnermostHandle }; }
+	constexpr FParameter Range(uint8 Num = 1) 	{ return {Name, NameLen, static_cast<uint8>(NumRanges + Num), bOptional, InnermostType, InnermostHandle }; }
 };
 
 template<Arithmetic T>
@@ -92,12 +105,64 @@ constexpr FParameter MakeParam(const char* Name, std::size_t Len)
 
 ////////////////////////////////////////////////////////////////////////////////////////////////
 
-constexpr FParameter operator ""_B1(const char* Name, std::size_t Len)	{ return MakeParam<bool>(Name, Len); }
-constexpr FParameter operator ""_S32(const char* Name, std::size_t Len)	{ return MakeParam<int32>(Name, Len); }
-constexpr FParameter operator ""_U32(const char* Name, std::size_t Len)	{ return MakeParam<uint32>(Name, Len); }
-constexpr FParameter operator ""_U64(const char* Name, std::size_t Len)	{ return MakeParam<uint64>(Name, Len); }
-constexpr FParameter operator ""_F32(const char* Name, std::size_t Len)	{ return MakeParam<float>(Name, Len); }
-constexpr FParameter operator ""_F64(const char* Name, std::size_t Len)	{ return MakeParam<double>(Name, Len); }
+inline constexpr FInnermostType InnermostDecomposedStruct	= {.Struct{._ = EInnermostKind::Struct, .IsDecomposed = 1}};
+inline constexpr FInnermostType InnermostDecomposedFlatEnum	= {.Enum{._ = EInnermostKind::Enum, .FlagMode = 0, .IsDecomposed = 1}};
+inline constexpr FInnermostType InnermostDecomposedFlagEnum	= {.Enum{._ = EInnermostKind::Enum, .FlagMode = 1, .IsDecomposed = 1}};
+
+template<class T>
+constexpr FInnermostType IllegalItem()
+{
+	static_assert(!sizeof(T), "Illegal item type");
+	return InnermostStaticStruct;
+}
+
+template<class T> inline constexpr FInnermostType ReflectItem				= IllegalItem<T>();
+template<Arithmetic T> inline constexpr FInnermostType ReflectItem<T>		= ReflectInnermost<T>;
+template<> inline constexpr FInnermostType ReflectItem<FOutBoolItem>		= ReflectInnermost<bool>;
+template<> inline constexpr FInnermostType ReflectItem<FInStruct>			= InnermostStaticStruct;
+template<> inline constexpr FInnermostType ReflectItem<FOutStruct>			= InnermostStaticStruct;
+template<> inline constexpr FInnermostType ReflectItem<FInAny>				= InnermostDynamicStruct;
+template<> inline constexpr FInnermostType ReflectItem<FOutAny>				= InnermostDynamicStruct;
+template<> inline constexpr FInnermostType ReflectItem<FInMembers>			= InnermostDecomposedStruct;
+template<> inline constexpr FInnermostType ReflectItem<FOutMembers>			= InnermostDecomposedStruct;
+template<> inline constexpr FInnermostType ReflectItem<FInEnum>				= InnermostFlatEnum; // Note that FlagMode is unknown
+template<> inline constexpr FInnermostType ReflectItem<FOutEnumItem>		= InnermostFlatEnum; // Note that FlagMode is unknown
+template<> inline constexpr FInnermostType ReflectItem<FInEnumerator>		= InnermostDecomposedFlatEnum;
+template<> inline constexpr FInnermostType ReflectItem<FOutEnumeratorItem>	= InnermostDecomposedFlatEnum;
+template<> inline constexpr FInnermostType ReflectItem<FInFlags>			= InnermostDecomposedFlagEnum;
+template<> inline constexpr FInnermostType ReflectItem<FOutFlagsItem>		= InnermostDecomposedFlagEnum;
+
+////////////////////////////////////////////////////////////////////////////////////////////////
+
+constexpr FParameter operator ""_Bool(const char* Name, std::size_t Len)	{ return MakeParam<bool>(Name, Len); }
+constexpr FParameter operator ""_B1(const char* Name, std::size_t Len)		{ return MakeParam<bool>(Name, Len); }
+constexpr FParameter operator ""_S8(const char* Name, std::size_t Len)		{ return MakeParam<int8>(Name, Len); }
+constexpr FParameter operator ""_U8(const char* Name, std::size_t Len)		{ return MakeParam<uint8>(Name, Len); }
+constexpr FParameter operator ""_S16(const char* Name, std::size_t Len)		{ return MakeParam<int16>(Name, Len); }
+constexpr FParameter operator ""_U16(const char* Name, std::size_t Len)		{ return MakeParam<uint16>(Name, Len); }
+constexpr FParameter operator ""_S32(const char* Name, std::size_t Len)		{ return MakeParam<int32>(Name, Len); }
+constexpr FParameter operator ""_U32(const char* Name, std::size_t Len)		{ return MakeParam<uint32>(Name, Len); }
+constexpr FParameter operator ""_U64(const char* Name, std::size_t Len)		{ return MakeParam<uint64>(Name, Len); }
+constexpr FParameter operator ""_F32(const char* Name, std::size_t Len)		{ return MakeParam<float>(Name, Len); }
+constexpr FParameter operator ""_F64(const char* Name, std::size_t Len)		{ return MakeParam<double>(Name, Len); }
+constexpr FParameter operator ""_Char(const char* Name, std::size_t Len)	{ return MakeParam<char8_t>(Name, Len); }
+constexpr FParameter operator ""_C8(const char* Name, std::size_t Len)		{ return MakeParam<char8_t>(Name, Len); }
+constexpr FParameter operator ""_C16(const char* Name, std::size_t Len)		{ return MakeParam<char16_t>(Name, Len); }
+constexpr FParameter operator ""_C32(const char* Name, std::size_t Len)		{ return MakeParam<char32_t>(Name, Len); }
+constexpr FParameter operator ""_B1s(const char* Name, std::size_t Len)		{ return MakeParam<bool>(Name, Len).Range(); }
+constexpr FParameter operator ""_S8s(const char* Name, std::size_t Len)		{ return MakeParam<int8>(Name, Len).Range(); }
+constexpr FParameter operator ""_U8s(const char* Name, std::size_t Len)		{ return MakeParam<uint8>(Name, Len).Range(); }
+constexpr FParameter operator ""_S16s(const char* Name, std::size_t Len)	{ return MakeParam<int16>(Name, Len).Range(); }
+constexpr FParameter operator ""_U16s(const char* Name, std::size_t Len)	{ return MakeParam<uint16>(Name, Len).Range(); }
+constexpr FParameter operator ""_S32s(const char* Name, std::size_t Len)	{ return MakeParam<int32>(Name, Len).Range(); }
+constexpr FParameter operator ""_U32s(const char* Name, std::size_t Len)	{ return MakeParam<uint32>(Name, Len).Range(); }
+constexpr FParameter operator ""_U64s(const char* Name, std::size_t Len)	{ return MakeParam<uint64>(Name, Len).Range(); }
+constexpr FParameter operator ""_F32s(const char* Name, std::size_t Len)	{ return MakeParam<float>(Name, Len).Range(); }
+constexpr FParameter operator ""_F64s(const char* Name, std::size_t Len)	{ return MakeParam<double>(Name, Len).Range(); }
+constexpr FParameter operator ""_Str(const char* Name, std::size_t Len)		{ return MakeParam<char8_t>(Name, Len).Range(); }
+constexpr FParameter operator ""_C8s(const char* Name, std::size_t Len)		{ return MakeParam<char8_t>(Name, Len).Range(); }
+constexpr FParameter operator ""_C16s(const char* Name, std::size_t Len)	{ return MakeParam<char16_t>(Name, Len).Range(); }
+constexpr FParameter operator ""_C32s(const char* Name, std::size_t Len)	{ return MakeParam<char32_t>(Name, Len).Range(); }
 
 ////////////////////////////////////////////////////////////////////////////////////////////////
 
@@ -114,6 +179,12 @@ struct FMemberLiteral
 
 ////////////////////////////////////////////////////////////////////////////////////////////////
 
+struct FOpaqueEnum
+{
+	FEnumId			Id;
+	bool			bFlagMode;
+};
+
 struct FEnumeratorNames
 {
 	FEnumId			Id;
@@ -128,22 +199,25 @@ struct FEnumFlagNames
 
 struct FEnumLiteral : FMemberLiteral
 {
-	inline FParameter operator()(FEnumId Id) const
+	uint8 NumRanges = 0;
+
+	inline FParameter operator()(FOpaqueEnum Enum) const
 	{
-		return MakeParam(0, InnermostFlatEnum, {FInnerId(Id), 0, 0, EAccess::WholeValue});
+		return MakeParam(NumRanges, Enum.bFlagMode ? InnermostFlagEnum : InnermostFlatEnum, {FInnerId(Enum.Id), 0, 0, EAccess::WholeValue});
 	}
 
 	inline FParameter operator()(FEnumeratorNames Enumerators, EAccess Access = EAccess::AllMembers) const
 	{
-		return MakeParam(0, InnermostFlatEnum, {FInnerId(Enumerators.Id), 0, Enumerators.Slice, Access});
+		return MakeParam(NumRanges, InnermostDecomposedFlatEnum, {FInnerId(Enumerators.Id), 0, Enumerators.Slice, Access});
 	}
 
 	inline FParameter operator()(FEnumFlagNames Flags, EAccess Access = EAccess::AllMembers) const
 	{
-		return MakeParam(0, InnermostFlagEnum, {FInnerId(Flags.Id), 0, Flags.Slice, Access});
+		return MakeParam(NumRanges, InnermostDecomposedFlagEnum, {FInnerId(Flags.Id), 0, Flags.Slice, Access});
 	}
 };
-inline constexpr FEnumLiteral operator ""_Enum(const char* Name, std::size_t Len)		{ return { {Name, Len} }; }
+inline constexpr FEnumLiteral operator ""_Enum(const char* Name, std::size_t Len)	{ return { {Name, Len}, 0 }; }
+inline constexpr FEnumLiteral operator ""_Enums(const char* Name, std::size_t Len)	{ return { {Name, Len}, 1 }; }
 
 ////////////////////////////////////////////////////////////////////////////////////////////////
 
@@ -155,103 +229,254 @@ struct FStructMemberNames
 
 struct FStructLiteral : FMemberLiteral
 {
+	uint8 NumRanges = 0;
+
 	inline FParameter operator()(FStructId Id, uint16 Version = 0) const
 	{
-		return MakeParam(0, InnermostStaticStruct, {FInnerId(Id), Version, 0, EAccess::WholeValue});
+		return MakeParam(NumRanges, InnermostStaticStruct, {FInnerId(Id), Version, 0, EAccess::WholeValue});
 	}
 
 	inline FParameter operator()(FStructMemberNames Members, uint16 Version = 0, EAccess Access = EAccess::AllMembers) const
 	{
-		return MakeParam(0, InnermostStaticStruct, {FInnerId(Members.Id), Version, Members.Slice, Access});
+		return MakeParam(NumRanges, InnermostDecomposedStruct, {FInnerId(Members.Id), Version, Members.Slice, Access});
+	}
+
+	inline FParameter Any() const
+	{
+		return MakeParam(NumRanges, InnermostDynamicStruct, {NoId, 0, 0, EAccess::WholeValue});
 	}
 };
-inline constexpr FStructLiteral operator ""_Struct(const char* Name, std::size_t Len)	{ return { Name, Len }; }
-inline constexpr FParameter operator ""_AnyStruct(const char* Name, std::size_t Len)	{ return { Name, static_cast<uint32>(Len), 0, false, InnermostDynamicStruct, {} }; }
+inline constexpr FStructLiteral operator ""_Struct(const char* Name, std::size_t Len)	{ return { {Name, Len}, 0 }; }
+inline constexpr FStructLiteral operator ""_Structs(const char* Name, std::size_t Len)	{ return { {Name, Len}, 1 }; }
 
 ////////////////////////////////////////////////////////////////////////////////////////////////
 
-struct FRangeLiteral : FMemberLiteral
+inline FLooseType ItemOf(FLooseType Type)
 {
-	uint8				NumRanges;
+	check(Type.NumRanges);
+	--Type.NumRanges;
+	return Type;
+}
 
-	template<Arithmetic T>
-	inline FParameter Of() const
-	{
-		return Make(ReflectInnermost<T>);
-	}
+////////////////////////////////////////////////////////////////////////////////////////////////
 
-	inline FParameter Of(FEnumId Id) const
-	{
-		return MakeParam(NumRanges, InnermostFlatEnum, {FInnerId(Id), 0, 0, EAccess::WholeValue});
-	}
+// Will allow logging errors/warnings, wont provide any global access to outside world
+// Might also be used to allocate new out ranges during op execution
+class FContext
+{
+	friend FFriend;
+	FScratchAllocator& Scratch;
+public:
+	explicit FContext(FScratchAllocator& In) : Scratch(In) {}
+};
 
-	inline FParameter Of(FEnumFlagNames Flags, EAccess Access = EAccess::AllMembers) const
-	{
-		return MakeParam(NumRanges, InnermostFlagEnum, {FInnerId(Flags.Id), 0, Flags.Slice, Access});
-	}
+////////////////////////////////////////////////////////////////////////////////////////////////
 
-	inline FParameter Of(FStructId Id, uint16 Version = 0) const
-	{
-		return MakeParam(NumRanges, InnermostStaticStruct, {FInnerId(Id), Version, 0, EAccess::WholeValue});
-	}
+struct FItems
+{
+	FLooseType			Type;
+	FLooseNestedRange	Value;
+	uint16				NumNamed;
+};
+
+template<typename T, uint8 N = 1>
+struct TInItems : protected FItems
+{
+	TInItems(FItems Items) : FItems(Items) { check(Items.Type.NumRanges == N-1); }
+	using Iterator = TInIt<T, N>::Type;
+	friend TOutRange<T,N>;
+
+	uint64				Num() const							{ return Value.Num; }
+	auto&&				operator[](uint64 Idx) const		{ check(Idx < Value.Num); return *Iterator(*this, Idx); }
+	Iterator			begin() const						{ return Iterator(*this, 0); }
+	Iterator			end() const							{ return Iterator(*this, Value.Num); }
+};
 
-	inline FParameter Of(FStructMemberNames Members, uint16 Version = 0, EAccess Access = EAccess::AllMembers) const
+struct FOutItems : FItems
+{
+	FLooseMember& Owner; // Only needed to RecordDynamicType
+};
+
+template<typename T, uint8 N = 1>
+struct TOutItems : protected FOutItems
+{
+	TOutItems(FOutItems Items) : FOutItems(Items) {}
+	using Iterator = TOutIt<T,N>::Type;
+	using DerefType = decltype(**((Iterator*)nullptr));
+
+	uint64				Num() const							{ return Value.Num; }
+	DerefType			operator[](uint64 Idx)				{ check(Idx < Value.Num); return *Iterator(*this, Idx); }
+	Iterator			begin() const						{ return Iterator(*this, 0); }
+	Iterator			end() const							{ return Iterator(*this, Value.Num); }
+};
+
+// Helps fill TOutItems and check that all items are assigned
+template<typename T, uint8 N = 1>
+class TFiller : protected TOutItems<T, N>
+{
+	using Super = TOutItems<T, N>;
+	uint64				Pos = 0;
+
+public:
+	TFiller(Super Items) : Super(Items) {}
+	~TFiller() { check(Pos == Super::Num()); }
+	
+	template<typename InT>
+	TFiller&			operator<<(InT InItem)
 	{
-		return MakeParam(NumRanges, InnermostStaticStruct, {FInnerId(Members.Id), Version, Members.Slice, Access});
+		Super::operator[](Pos++) = InItem;
+		return *this;
 	}
-
-	inline FParameter OfAny() const
+	
+	template<typename InT>
+	TFiller&			operator<<(TInItems<InT, N> Items)
 	{
-		return MakeParam(NumRanges, InnermostDynamicStruct, {NoId, 0, 0, EAccess::WholeValue});
+		check(Pos + Items.Num() <= Super::Num());
+		typename Super::Iterator OutIt(*this, Pos);
+		for (auto Item : Items)
+		{
+			*OutIt = Item;
+			++OutIt;
+		}
+		Pos += Items.Num();
+		return *this;
 	}
+};
+
+////////////////////////////////////////////////////////////////////////////////////////////////
+
+struct FInRange
+{
+	FLooseType			Type;
+	FLooseValue			Value;
+	uint16				NumNamed;
+
+	FItems				Items() const { return {ItemOf(Type), {Value.Meta.Num, Value.Data.Range}, NumNamed}; }
+};
+
+#if DO_CHECK
+PLAINPROPS_API void CheckRangeType(FLooseType Type, uint16 NumNamed, FInnermostType InnermostType, uint8 NumRanges);
+#else
+inline void CheckRangeType(FLooseType Type, uint16 NumNamed, FInnermostType InnermostType, uint8 NumRanges) {}
+#endif
+
+
+template<typename T, uint8 N = 1>
+struct TInRange : protected FInRange
+{
+	friend TOutRange<T,N>;
+	TInRange(FLooseMember I, uint16 NumNamed) : FInRange({I.GetType(), I.GetValue(), NumNamed})	{ CheckRangeType(Type, NumNamed, ReflectItem<T>, N);	}
+	uint64							Num() const			{ return Value.Meta.Num; }
+	TInItems<T, N>					Items() const		{ return FInRange::Items(); }
+};
+
+template<typename T, uint8 N = 1>
+struct TOptionalInRange : protected FInRange
+{
+	friend TOutRange<T,N>;
+	TOptionalInRange(FLooseMember I, uint16 NumNamed) : FInRange({I.GetType(), I.GetValue(), NumNamed})	{ CheckRangeType(Type, NumNamed, ReflectItem<T>, N); }
+	explicit						operator bool()		{ return Value; }
+	uint64							Num() const			{ return Value.Meta.Num; }
+	TInItems<T, N>					Items() const		{ check(Value); return FInRange::Items(); }
+};
+
+////////////////////////////////////////////////////////////////////////////////////////////////
 
-	inline FRangeLiteral Nest(uint8 NumNested = 1) const
+struct FOutRange
+{
+	FLooseMember&					Out;
+	const uint16					NumNamed;
+
+	FOutRange(FLooseMember& O, uint16 N) : Out(O), NumNamed(N) {}
+
+	PLAINPROPS_API void				Set(const FInRange& In);
+	PLAINPROPS_API void				Set(FLooseNestedRange Value);
+	PLAINPROPS_API FRangeData		NewNested(uint64 Num, FContext Ctx) const;
+	PLAINPROPS_API FRangeData		NewStructs(uint64 Num, FContext Ctx)  const;
+	PLAINPROPS_API FRangeData		NewMembers(uint64 Num, FContext Ctx) const;
+	PLAINPROPS_API FRangeData		NewEnums(uint64 Num, FContext Ctx) const;
+	PLAINPROPS_API FRangeData		NewEnumerators(uint64 Num, FContext Ctx) const;
+	PLAINPROPS_API FRangeData		NewFlags(uint64 Num, FContext Ctx) const;
+	PLAINPROPS_API FRangeData		NewBools(uint64 Num, FContext Ctx) const;
+	PLAINPROPS_API FRangeData		NewArithmetics(uint64 Num, FContext Ctx, uint32 SizeOf) const;
+	template<typename T> FRangeData	NewItems(uint64 Num, FContext Ctx) const;
+	FOutItems						SetItems(FLooseNestedRange Value)
 	{
-		check(NumRanges == 1);
-		check(NumNested > 0 && NumNested < MaxRangeNesting);
-		return { {Str, Len}, static_cast<uint8>(NumRanges + NumNested) };
+		Set(Value);
+		return {{ItemOf(Out.GetType()), Value, NumNamed}, Out};
 	}
 };
-consteval FRangeLiteral operator ""_Range(const char* Name, std::size_t Len) { return { {Name, Len}, 1 }; }
 
-////////////////////////////////////////////////////////////////////////////////////////////////
+template<typename T>
+FRangeData FOutRange::NewItems(uint64 Num, FContext Ctx) const
+{
+	if constexpr (std::is_same_v<T, bool>)						return NewBools(Num, Ctx);
+	else if constexpr (Arithmetic<T>)							return NewArithmetics(Num, Ctx, sizeof(T));
+	else if constexpr (std::is_same_v<T, FOutStruct>)			return NewStructs(Num, Ctx);
+	else if constexpr (std::is_same_v<T, FOutAny>)				return NewStructs(Num, Ctx);
+	else if constexpr (std::is_same_v<T, FOutMembers>)			return NewMembers(Num, Ctx);
+	else if constexpr (std::is_same_v<T, FOutEnumItem>)			return NewEnums(Num, Ctx);
+	else if constexpr (std::is_same_v<T, FOutEnumeratorItem>)	return NewEnumerators(Num, Ctx);
+	else if constexpr (std::is_same_v<T, FOutFlagsItem>)		return NewFlags(Num, Ctx);
+	else static_assert(!sizeof(T), "Illegal item type");
+}
 
-// Used to access inner struct members
-template<uint32 N>
-class TIns
+template<typename T, uint8 N>
+struct TOutRange : protected FOutRange
 {
-public:
-	TIns(const FIn* InIns) : Ins(InIns) { static_assert(N > 1); }
-	const FIn& operator[](uint32 Idx) const;
-private:
-	const FIn* Ins;
+	TOutRange(FLooseMember& O, uint16 NumNamed) : FOutRange(O, NumNamed) { CheckRangeType(Out.GetType(), NumNamed, ReflectItem<T>, N); }
+	
+	void				operator=(const TInRange<T, N>& In)	{ Set(In); }
+	void				operator=(TInItems<T, N> Items)	{ Set(Items); }
+	TOutItems<T, N>		New(uint64 Num, FContext Ctx) requires (N>1)	{ return SetItems({Num, NewNested(Num, Ctx)}); }
+	TOutItems<T, N>		New(uint64 Num, FContext Ctx) requires (N==1)	{ return SetItems({Num, NewItems<T>(Num, Ctx)}); }
 };
 
+////////////////////////////////////////////////////////////////////////////////////////////////
+
 class FIn : protected FLooseMember
 {
 public:
 	using FLooseMember::FLooseMember;
+	FIn(const FLooseMember& In) : FLooseMember(In.GetType(), In.GetValue()) {}
 	using FLooseMember::operator bool;
 
-	bool							Bool() const { return As<bool>(); }
-	int32							S32() const { return As<int32>(); }
-	uint32							U32() const { return As<uint32>(); }
-	uint64							U64() const { return As<uint64>(); }
-	float							F32() const { return As<float>(); }
-	PLAINPROPS_API FInEnum			Enum() const; 
-	PLAINPROPS_API FInEnumerator	Enumerator() const;
-	PLAINPROPS_API FInEnumFlags		EnumFlags() const;
-	PLAINPROPS_API FInStruct		Struct() const;
-
-
-	TOptional<bool>					OptBool() const { return AsOpt<bool>(); }
-	TOptional<int32>				OptS32() const { return AsOpt<int32>(); }
-	TOptional<uint32>				OptU32() const { return AsOpt<uint32>(); }
-	TOptional<uint64>				OptU64() const { return AsOpt<uint64>(); }
-	TOptional<float>				OptF32() const { return AsOpt<float>(); }
-	PLAINPROPS_API FInStruct		OptStruct() const;
-
-	template<uint32 N> TIns<N>		StructMembers() const { return TIns<N>(AsStructMembers(N)); }
+	bool												Bool() const							{ return As<bool>(); }
+	int32												S32() const								{ return As<int32>(); }
+	uint32												U32() const								{ return As<uint32>(); }
+	uint64												U64() const								{ return As<uint64>(); }
+	float												F32() const								{ return As<float>(); }
+	char8_t												Char() const							{ return As<char8_t>(); }
+	char8_t												C8() const								{ return As<char8_t>(); }
+
+	PLAINPROPS_API FInStruct							Struct() const;
+	PLAINPROPS_API FInAny								StructAny() const;
+	PLAINPROPS_API FInMembers							StructMembers() const;
+
+	PLAINPROPS_API FInEnum								Enum() const; 
+	PLAINPROPS_API FInEnumerator						Enumerator() const;
+	PLAINPROPS_API FInFlags								EnumFlags() const;
+
+	template<typename T, uint8 N = 1>
+	TInRange<T, N>										Range(uint16 NumNamed = 0) const		{ return TInRange<T, N>(*this, NumNamed); };
+	TInRange<bool>										Bools() const							{ return Range<bool>(); }
+	TInRange<uint8>										U8s() const								{ return Range<uint8>(); }
+	TInRange<int32>										S32s() const							{ return Range<int32>(); }
+	TInRange<float>										F32s() const							{ return Range<float>(); }
+	TInRange<char8_t>									C8s() const								{ return Range<char8_t>(); }
+	TInRange<char8_t>									Str() const								{ return Range<char8_t>(); }
+	TInRange<FInStruct>									Structs() const							{ return Range<FInStruct>(); }
+	TInRange<FInMembers>								StructsMembers(uint16 NumNamed) const	{ return Range<FInMembers>(NumNamed); }
+	TInRange<FInEnum>									Enums() const							{ return Range<FInEnum>(); }
+	TInRange<FInEnumerator>								Enumerators(uint16 NumNamed) const		{ return Range<FInEnumerator>(NumNamed); }
+	TInRange<FInFlags>									EnumsFlags(uint16 NumNamed) const		{ return Range<FInFlags>(NumNamed); }
+
+	TOptional<bool>										OptBool() const							{ return AsOpt<bool>(); }
+	TOptional<int32>									OptS32() const							{ return AsOpt<int32>(); }
+	TOptional<uint32>									OptU32() const							{ return AsOpt<uint32>(); }
+	TOptional<uint64>									OptU64() const							{ return AsOpt<uint64>(); }
+	TOptional<float>									OptF32() const							{ return AsOpt<float>(); }
+	PLAINPROPS_API FInStruct							OptStruct() const;
 
 private:	
 	template<Arithmetic T>
@@ -269,27 +494,37 @@ private:
 		check(IsOptional());
 		return Value ? TOptional<T>(ArithmeticCast<T>(Value.Data.Arithmetic)) : NullOpt;
 	}
+};
 
-	PLAINPROPS_API const FIn*		AsStructMembers(uint32 N) const;
+// Opaque struct value, assignable to a type-matching FOutStruct
+struct FInStruct : protected FLooseMember
+{
+	using FLooseMember::FLooseMember;
+	friend FOutStruct; friend FOutAny;
 };
 
-template<uint32 N>
-const FIn& TIns<N>::operator[](uint32 Idx) const
+struct FInAny : protected FLooseMember
 {
-	check(Idx < N);
-	return Ins[Idx];
-}
+	using FLooseMember::FLooseMember;
+	friend FOutStruct;
+};
 
-// Opaque struct value, assignable to a type-matching FOutStruct
-class FInStruct : protected FIn
+struct FInMembers : protected FLooseMember
 {
+	using FLooseMember::FLooseMember;
 	friend FIn; friend FOutStruct;
+
+	FIn operator[](uint32 Idx) const { check(Idx < Value.Meta.Num); return FIn(Value.Data.Struct.Loose[Idx]); }
 };
 
 // Opaque enum value, assignable to a type-matching FOutEnum
-class FInEnum : protected FIn
+class FInEnum
 {
-	friend FIn; friend FOutEnum;
+	friend FOutEnum; friend FOutEnumItem;
+	FLooseType			Type;
+	FEnumeratorData		Value;
+public:
+	FInEnum(FLooseType T, FEnumeratorData V) : Type(T), Value(V) {}
 };
 
 // Enum value represented as enumerator index in NameEnumerators() order
@@ -298,24 +533,23 @@ class FInEnum : protected FIn
 // * Not convertible to FInEnum nor assignable to FOutEnum either
 class FInEnumerator
 {
-	friend FIn;
-	uint8 Index;
+	FEnumeratorIndex	Name;
 public:
-	explicit FInEnumerator(uint8 Idx) : Index(Idx) {}
-
-	uint8			GetIndex() const { return Index;}
+	explicit FInEnumerator(FEnumeratorIndex N, uint16 Max) : Name(N) { check(Name.Index < Max); }
+	uint16				GetIndex() const { return Name.Index; }
 };
 
 // Bitset of enumerator indices in NameEnumFlags() order
 // * Conceptually a string set
 // * Not assignable to FOutEnumerator, which might use different NameEnumerators() indices
-class FInEnumFlags
+class FInFlags
 {
-	uint64 Indices;
+	FEnumeratorIndexSet	Flags;
+	const uint32		Max;
 public:
-	explicit FInEnumFlags(uint64 IndexFlags) : Indices(IndexFlags) {}
-	uint64			GetIndices() const;
-	bool			IsSet(uint8 Index) const;
+	explicit FInFlags(FEnumeratorIndexSet InFlags, uint32 NumNamedFlags) : Flags(InFlags), Max(NumNamedFlags) {}
+	uint64				GetIndices() const;
+	bool				IsSet(uint32 Index) const;
 };
 
 ////////////////////////////////////////////////////////////////////////////////////////////////
@@ -333,6 +567,7 @@ public:
 
 	void operator=(T In)
 	{
+		// todo: normalize floats or check is normalized
 		Out.Meta.IsSet = 1;
 		Out.Data.Arithmetic = ValueCast(In);
 	}
@@ -344,162 +579,336 @@ public:
 	}
 };
 
-// TOutRange<Arithmetic> helper
-template <Arithmetic T>
-class TNormalizer
+class FOut : protected FLooseMember
 {
-	T& Out;
 public:
-	TNormalizer(T& O) : Out(O) {}
-	void operator=(T In) { Out = Normalize(In);	}
+	FOut(FOut&&) = default;
+	FOut(const FOut&) = delete;
+	FOut(FOutParameter Parameter);
+	FOut(FOutParameter Parameter, FOut* StructMembers);
+	FOut& operator=(FOut&&) = default;
+	FOut& operator=(const FOut&) = delete;
+
+	TOut<bool>						Bool()							{ return { Type, Value }; }
+	TOut<int32>						S32()							{ return { Type, Value }; }
+	TOut<uint32>					U32()							{ return { Type, Value }; }
+	TOut<uint64>					U64()							{ return { Type, Value }; }
+	TOut<float>						F32()							{ return { Type, Value }; }
+	TOut<char8_t>					Char()							{ return { Type, Value }; }
+	TOut<char8_t>					C8()							{ return { Type, Value }; }
+	
+	PLAINPROPS_API FOutStruct		Struct();
+	PLAINPROPS_API FOutAny			StructAny();
+	PLAINPROPS_API FOutMembers		StructMembers();
+
+	PLAINPROPS_API FOutEnum			Enum();
+	PLAINPROPS_API FOutEnumerator	Enumerator();
+	PLAINPROPS_API FOutFlags		EnumFlags();
+
+	template<typename T, uint8 N = 1>
+	TOutRange<T, N>					Range(uint16 NumNamed = 0)		{ return TOutRange<T, N>(*this, NumNamed); }
+	
+	TOutRange<FOutBoolItem>			Bools()							{ return Range<FOutBoolItem>(); }
+	TOutRange<uint8>				U8s()							{ return Range<uint8>(); }
+	TOutRange<int32>				S32s()							{ return Range<int32>(); }
+	TOutRange<float>				F32s()							{ return Range<float>(); }
+	TOutRange<char8_t>				C8s()							{ return Range<char8_t>(); }
+	TOutRange<char8_t>				Str()							{ return Range<char8_t>(); }
+	TOutRange<FOutStruct>			Structs()						{ return Range<FOutStruct>(); }
+	TOutRange<FOutAny>				StructsAny()					{ return Range<FOutAny>(); }
+	TOutRange<FOutMembers>			StructsMembers(uint16 NumNamed) { return Range<FOutMembers>(NumNamed); }
+	TOutRange<FOutEnumItem>			Enums()							{ return Range<FOutEnumItem>(); }
+	TOutRange<FOutEnumeratorItem>	Enumerators(uint16 NumNamed)	{ return Range<FOutEnumeratorItem>(NumNamed); }
+	TOutRange<FOutFlagsItem>		EnumsFlags(uint16 NumNamed)		{ return Range<FOutFlagsItem>(NumNamed); }
+
+	friend FFriend; friend FOutMembers;
 };
 
-template <Arithmetic T>
-class TOutRange
+struct FOutMember
 {
-	TArrayView64<T> Out;
-public:
-	TOutRange(uint64 Num, T* Data) : Out(Data, Num) {}
-	TNormalizer<T>&	operator[](int64 Idx) { return Out[Idx]; }
+	const FLooseType		Type;
+	FLooseValue&			Value;
+};
+
+struct FOutStruct : protected FOutMember
+{
+	FOutStruct(FLooseType Type, FLooseValue& Out) : FOutMember{Type, Out} {}
+	PLAINPROPS_API void operator=(FInStruct In);
 };
 
-// TOutRange<bool> helper
-class FBitSetter
+class FOutAny : protected FOutMember
 {
-	uint8&				Out;
-	const uint8			Mask;
+	FLooseMember&			Owner;
 public:
-	FBitSetter(uint8& Byte, uint8 Idx) : Out(Byte), Mask(uint8(1) << Idx) {}
+	FOutAny(FLooseType Type, FLooseValue& Out, FLooseMember& InOwner) : FOutMember{Type, Out}, Owner(InOwner) {}
+	PLAINPROPS_API void operator=(FInStruct In);
+};
 
-	void operator=(bool bValue)
-	{
-		if (bValue)
-		{
-			Out |= Mask;
-		}
-		else
-		{
-			Out &= ~Mask;
-		}	
-	}
+struct FOutMembers : protected FOutMember
+{
+	FOutMembers(FLooseType Type, FLooseValue& Out) : FOutMember{Type, Out} {}
+	FOut& operator[](uint32 Idx) { check(Idx < Value.Meta.Num);	return static_cast<FOut&>(Value.Data.Struct.Loose[Idx]); }
 };
 
-template<>
-class TOutRange<bool>
+struct FOutEnum : protected FOutMember
 {
-	const uint64					Num;
-	uint8*							Out;
-public:
-	TOutRange(uint64 N, uint8* Bits) : Num(N), Out(Bits) {}
-	FBitSetter operator[](uint64 Idx) { check(Idx < Num); return FBitSetter(Out[Idx/8], Idx%8); }
+	FOutEnum(FLooseType Type, FLooseValue& Out) : FOutMember{Type, Out} {}
+	PLAINPROPS_API void operator=(FInEnum In);
 };
 
-class FOutStructs
+struct FOutEnumerator : protected FOutMember
 {
-	const FDeclId					Id;
-	const uint16					Version;
-	TArrayView<FLooseValue>			Out;
-public:
-	FOutStructs(FDeclId I, uint16 V, TArrayView<FLooseValue> O) : Id(I), Version(V), Out(O) {}
+	FOutEnumerator(FLooseType Type, FLooseValue& Out) : FOutMember{Type, Out} {}
+	PLAINPROPS_API void operator=(uint16 Index);
 };
 
-class FOutRange
+class FOutFlag
 {
+	FLooseValue&			Out;
+	const uint32			Idx;
 public:
-	FOutRange(FLooseType T, uint64 N, void* D) : Type(T), Num(N), Data(D) {}
-	
-	TOutRange<int32>				S32s() { return Of<int32>(); }
-	TOutRange<uint32>				U32s() { return Of<uint32>(); }
-	TOutRange<uint64>				U64s() { return Of<uint64>(); }
-	TOutRange<float>				F32s() { return Of<float>(); }
-	FOutStructs						Structs();// { check(Type.IsStruct()); return Items; }
-	FOutStructs						DynamicStructs(FDeclId Id, uint16 Version);// { check(Type.IsStruct() && Type.InnermostType.) }
+	FOutFlag(FLooseValue& O, uint32 I) : Out(O), Idx(I) { check(Idx < Out.Meta.Num); }
+	PLAINPROPS_API void operator=(bool bValue);
+	PLAINPROPS_API void operator=(TOptional<bool> bValue);
+};
 
-private:
-	const FLooseType				Type;
-	uint64							Num;
-	void*							Data;
+struct FOutFlags : protected FOutMember
+{
+	FOutFlags(FLooseType Type, FLooseValue& Out) : FOutMember{Type, Out} {}
+	// @param Index in NameFlags order
+	FOutFlag operator[](uint32 Idx) { return FOutFlag(Value, Idx); }
+	// @param Indices bitset in NameFlags order
+	PLAINPROPS_API void operator=(uint64 Indices);
+};
 
-	template<Arithmetic T>
-	TOutRange<T>					Of()
-	{
-		check(Type.Compare(ReflectLoose<T>));
-		return TOutRange<T>(Num, static_cast<T*>(Data));
-	}
+////////////////////////////////////////////////////////////////////////////////////////////////
+
+class FOutBoolItem
+{
+	uint8&					Out;
+	uint8					Mask;
+public:
+	FOutBoolItem(uint8& Byte, uint8 InMask) : Out(Byte), Mask(InMask) {}
+	PLAINPROPS_API void operator=(bool bValue);
 };
 
-class FOut : protected FLooseMember
+class FOutEnumItem
 {
+	const FLooseType		Type;
+	FEnumeratorData&		Out;
 public:
-	FOut(FOut&&) = default;
-	FOut(const FOut&) = delete;
-	FOut(FOutParameter Parameter);
-	FOut(FOutParameter Parameter, FOut* StructMembers);
-	FOut& operator=(FOut&&) = default;
-	FOut& operator=(const FOut&) = delete;
+	FOutEnumItem(FLooseType T, FEnumeratorData& D) : Type(T), Out(D) {}
+	PLAINPROPS_API void operator=(FInEnum In);
+};
 
-	TOut<int32>						S32()		{ return { Type, Value }; }
-	TOut<uint32>					U32()		{ return { Type, Value }; }
-	TOut<uint64>					U64()		{ return { Type, Value }; }
-	TOut<float>						F32()		{ return { Type, Value }; }
-	
-	PLAINPROPS_API FOutEnum&		Enum();
-	PLAINPROPS_API FOutEnumerator&	Enumerator();
-	PLAINPROPS_API FOutEnumFlags&	EnumFlags();
-	PLAINPROPS_API FOutStruct&		Struct();
-	template<uint32 N> TOuts<N>		StructMembers() { return AsStructMembers(N); }
+class FOutEnumeratorItem
+{
+	FEnumeratorIndex&		Out;
+	uint32					Max;
+public:
+	FOutEnumeratorItem(FEnumeratorIndex& Name, uint32 InMax) : Out(Name), Max(InMax) {}
+	PLAINPROPS_API void operator=(uint16 Index);
+};
 
-	PLAINPROPS_API FOutRange		NewRange(uint64 Num); // Allocate zeroed range (note: copy Num down for e.g. range of struct members)
+class FOutFlagItem
+{
+	FEnumeratorIndexSet&	Out;
+	const uint32			Idx;
+public:
+	FOutFlagItem(FEnumeratorIndexSet& O, uint32 InIdx) : Out(O), Idx(InIdx) {}
+	PLAINPROPS_API void operator=(bool bValue);
+};
 
-private:
-	PLAINPROPS_API FOut*			AsStructMembers(uint32 N);
+class FOutFlagsItem
+{
+	FEnumeratorIndexSet&	Out;
+	const uint32			Max;
+public:
+	FOutFlagsItem(FEnumeratorIndexSet& Flags, uint32 InMax) : Out(Flags), Max(InMax) {}
+	// @param Index in NameFlags order
+	FOutFlagItem operator[](uint32 Index) { check(Index < Max); return FOutFlagItem(Out, Index); }
+	// @param Indices bitset in NameFlags order
+	PLAINPROPS_API void operator=(uint64 Indices);
+};
+
+////////////////////////////////////////////////////////////////////////////////////////////////
 
-	friend const FLooseMember&		ReadOutput(const FOut& Out);
-	friend TArrayView<FLooseMember>	EditOutputs(TArrayView<FOut> Outs);
+struct FInBoolIt : public FBoolRangeIterator
+{
+	FInBoolIt(FItems Items, uint64 Idx) : FBoolRangeIterator(static_cast<const uint8*>(Items.Value.Data.Arithmetics.Source), Idx) {}
 };
 
-class FOutEnumFlag
+class FOutBoolIt
 {
-	FLooseValue& Out;
-	const uint64 Mask;
+	uint8*				Byte;
+	uint8				Mask;
+	PLAINPROPS_API void	Increment();
 public:
-	FOutEnumFlag(FLooseValue& O, uint32 Idx) : Out(O), Mask(uint64(1) << Idx) { check(Idx < Out.Meta.Num); }
-	PLAINPROPS_API void operator=(bool bValue);
-	PLAINPROPS_API void operator=(TOptional<bool> bValue);
+	FOutBoolIt(FItems Items, uint64 Idx)
+	: Byte(static_cast<uint8*>(Items.Value.Data.Arithmetics.Loose) + Idx / 8)
+	, Mask(1u << (Idx % 8))
+	{}
+
+	FOutBoolItem		operator*() const { return FOutBoolItem(*Byte, Mask); }
+	bool				operator==(const FOutBoolIt& Rhs) const = default;
+	bool				operator!=(const FOutBoolIt& Rhs) const = default;
+	FOutBoolIt&			operator++() { Increment(); return *this; }
 };
 
-class FOutEnumFlags : protected FOut
+template<Arithmetic T>
+class TInArithmeticIt 
 {
+	const T*			Ptr;
 public:
-	FOutEnumFlag operator[](uint32 Idx) { return FOutEnumFlag(/* out */ Value, Idx); }
+	TInArithmeticIt(FItems Items, uint64 Idx) : Ptr(static_cast<const T*>(Items.Value.Data.Arithmetics.Source) + Idx) {}
+	T					operator*()								{ return *Ptr; }
+	TInArithmeticIt&	operator++()							{ ++Ptr; return *this; }
+	bool				operator!=(TInArithmeticIt O) const		{ return Ptr != O.Ptr; }
 };
 
-class FOutEnum : protected FOut
+template<Arithmetic T>
+class TOutArithmeticIt
 {
+	T*					Ptr;
 public:
-	PLAINPROPS_API void operator=(FInEnum In);
+	TOutArithmeticIt(FItems Items, uint64 Idx) : Ptr(static_cast<T*>(Items.Value.Data.Arithmetics.Loose) + Idx) {}
+	T&					operator*()								{ return *Ptr; }
+	TOutArithmeticIt&	operator++()							{ ++Ptr; return *this; }
+	bool				operator!=(TOutArithmeticIt O) const	{ return Ptr != O.Ptr; }
 };
 
-class FOutStruct : protected FOut
+template<typename ValueType>
+class TItemIt
 {
 public:
-	PLAINPROPS_API void operator=(FInStruct In);
+	TItemIt&			operator++()						{ ++It; return *this; }
+	bool				operator!=(TItemIt O) const			{ return It != O.It; }
+protected:
+	TItemIt(FItems Items, ValueType* Value) : Type(Items.Type), NumNamed(Items.NumNamed), It(Value) {}
+
+	const FLooseType	Type;
+	uint16				NumNamed;
+	ValueType*			It;
 };
 
-template<uint32 N>
-class TOuts
+struct FInStructIt : TItemIt<const FLooseValue>
 {
-public:
-	TOuts(FOut* InOuts) : Outs(InOuts) {}
-	FOut& operator[](uint32 Idx) const { check(Idx < N); return Outs[Idx]; }
+	FInStructIt(FItems Items, uint64 Idx) : TItemIt(Items, Items.Value.Data.Structs.Loose + Idx) { checkSlow(Type.IsStruct() && NumNamed == 0); }
+	FInStruct operator*() const { return FInStruct(Type, *It); }
+};
+
+struct FOutStructIt : TItemIt<FLooseValue>
+{
+	FOutStructIt(FItems Items, uint64 Idx) : TItemIt(Items, Items.Value.Data.Structs.Loose + Idx) { checkSlow(Type.IsStruct() && NumNamed == 0); }
+	FOutStruct operator*() const { return FOutStruct(Type, *It); }
+};
+
+struct FInAnyIt : TItemIt<const FLooseValue>
+{
+	FInAnyIt(FItems Items, uint64 Idx) : TItemIt(Items, Items.Value.Data.Structs.Loose + Idx) {}
+	FInAny operator*() const { return FInAny(Type, *It); }
+};
+
+struct FOutAnyIt : TItemIt<FLooseValue>
+{
+	FOutAnyIt(FOutItems Items, uint64 Idx) : TItemIt(Items, Items.Value.Data.Structs.Loose + Idx), Owner(Items.Owner) { checkSlow(Type.IsStruct() && NumNamed == 0); }
+	FOutAny operator*() const { return FOutAny(Type, *It, Owner); }
 private:
-	FOut* const Outs;
+	FLooseMember& Owner;
 };
 
-////////////////////////////////////////////////////////////////////////////////////////////////
+struct FInMemberIt : TItemIt<const FLooseValue>
+{
+	FInMemberIt(FItems Items, uint64 Idx) : TItemIt(Items, Items.Value.Data.Structs.Loose + Idx) { checkSlow(Type.IsStruct() && NumNamed > 0); }
+	FInMembers operator*() const { return FInMembers(Type, *It); }
+};
 
-// Only for logging errors and warnings, todo
-class FContext {};
+struct FOutMemberIt : TItemIt<FLooseValue>
+{
+	FOutMemberIt(FItems Items, uint64 Idx) : TItemIt(Items, Items.Value.Data.Structs.Loose + Idx) { checkSlow(Type.IsStruct() && NumNamed > 0);}
+	FOutMembers operator*() const { return FOutMembers(Type, *It); }
+};
+
+struct FInEnumIt : TItemIt<const FEnumLooseData>
+{
+	FInEnumIt(FItems Items, uint64 Idx) : TItemIt(Items, Items.Value.Data.Enums.Loose + Idx) { checkSlow(Type.IsEnum() && NumNamed == 0);}
+	FInEnum operator*() const { return FInEnum(Type, It->Enumerator); }
+};
+
+struct FOutEnumIt : TItemIt<FEnumLooseData>
+{
+	FOutEnumIt(FItems Items, uint64 Idx) : TItemIt(Items, Items.Value.Data.Enums.Loose + Idx) { checkSlow(Type.IsEnum() && NumNamed == 0);}
+	FOutEnumItem operator*() const { return FOutEnumItem(Type, It->Enumerator); }
+};
+
+struct FInEnumeratorIt : TItemIt<const FEnumLooseData>
+{
+	FInEnumeratorIt(FItems Items, uint64 Idx) : TItemIt(Items, Items.Value.Data.Enums.Loose + Idx) { checkSlow(Type.InnermostType == InnermostDecomposedFlatEnum && !Type.IsRange() && NumNamed > 0); }
+	FInEnumerator operator*() const { return FInEnumerator(It->Decomposed.Flat, NumNamed); }
+};
+
+struct FOutEnumeratorIt : TItemIt<FEnumLooseData>
+{
+	FOutEnumeratorIt(FItems Items, uint64 Idx) : TItemIt(Items, Items.Value.Data.Enums.Loose + Idx) { checkSlow(Type.InnermostType == InnermostDecomposedFlatEnum && !Type.IsRange() && NumNamed > 0 && NumNamed < 64); }
+	FOutEnumeratorItem operator*() const { return FOutEnumeratorItem(It->Decomposed.Flat, NumNamed); }
+};
+
+struct FInFlagsIt : TItemIt<const FEnumLooseData>
+{
+	FInFlagsIt(FItems Items, uint64 Idx) : TItemIt(Items, Items.Value.Data.Enums.Loose + Idx) { checkSlow(Type.InnermostType == InnermostDecomposedFlagEnum && !Type.IsRange() && NumNamed > 0 && NumNamed < 64); }
+	FInFlags operator*() const { return FInFlags(It->Decomposed.Flags, static_cast<uint8>(NumNamed)); }
+};
+
+struct FOutFlagsIt : TItemIt<FEnumLooseData>
+{
+	FOutFlagsIt(FItems Items, uint64 Idx) : TItemIt(Items, Items.Value.Data.Enums.Loose + Idx) { checkSlow(Type.InnermostType == InnermostDecomposedFlagEnum && !Type.IsRange() && NumNamed > 0 && NumNamed < 64); }
+	FOutFlagsItem operator*() const { return FOutFlagsItem(It->Decomposed.Flags, NumNamed); }
+};
+
+template<typename T, uint8 N>
+struct TInNestedRangeIt : TItemIt<const FLooseNestedRange>
+{
+	TInNestedRangeIt(FItems Items, uint64 Idx) : TItemIt(Items, Items.Value.Data.Ranges + Idx) { checkSlow(Type.NumRanges == N-1); }
+	TInItems<T, N-1> operator*() const { return FItems{ItemOf(Type), *It, NumNamed}; }
+};
+
+template<typename T, uint8 N>
+struct TOutNestedRangeIt : TItemIt<FLooseNestedRange>
+{
+	TOutNestedRangeIt(FItems Items, uint64 Idx) : TItemIt(Items, Items.Value.Data.Ranges + Idx) { checkSlow(Type.NumRanges == N-1); }
+	TOutItems<T, N-1> operator*() const { return FItems{ItemOf(Type), *It, NumNamed}; }
+};
+
+template<typename T, uint8 N> requires(N==1) struct TInIt<T,N>	{ static_assert(!sizeof(T), "Illegal in item type"); };
+template<typename T, uint8 N> requires(N==1) struct TOutIt<T,N>	{ static_assert(!sizeof(T), "Illegal out item type"); };
+template<typename T, uint8 N> requires(N>1) struct TInIt<T,N>	{ using Type = TInNestedRangeIt<T,N>; }; 
+template<typename T, uint8 N> requires(N>1) struct TOutIt<T,N> 	{ using Type = TOutNestedRangeIt<T,N>; };
+template<Arithmetic T> struct TInIt<T,1>						{ using Type = TInArithmeticIt<T>; };
+template<Arithmetic T> struct TOutIt<T,1>						{ using Type = TOutArithmeticIt<T>; };
+template<> struct TInIt<bool,1>									{ using Type = FInBoolIt; };
+template<> struct TOutIt<FOutBoolItem,1>						{ using Type = FOutBoolIt; };
+template<> struct TInIt<FInStruct,1>							{ using Type = FInStructIt; };
+template<> struct TOutIt<FOutStruct,1>							{ using Type = FOutStructIt; };
+template<> struct TInIt<FInAny,1>								{ using Type = FInAnyIt; };
+template<> struct TOutIt<FOutAny,1>								{ using Type = FOutAnyIt; };
+template<> struct TInIt<FInMembers,1>							{ using Type = FInMemberIt; };
+template<> struct TOutIt<FInMembers,1>							{ using Type = FOutMemberIt; };
+template<> struct TInIt<FInEnum,1>								{ using Type = FInEnumIt; };
+template<> struct TOutIt<FOutEnumItem,1>						{ using Type = FOutEnumIt; };
+template<> struct TInIt<FInEnumerator,1>						{ using Type = FInEnumeratorIt; };
+template<> struct TOutIt<FOutEnumeratorItem,1>					{ using Type = FOutEnumeratorIt; };
+template<> struct TInIt<FInFlags,1>								{ using Type = FInFlagsIt; };
+template<> struct TOutIt<FOutFlagsItem,1>						{ using Type = FOutFlagsIt; };
+
+// Remap TOutItems/TFiller<Xyz> ->  TOutItems/TFiller<XyzItem> to enable e.g. TOutItems<FOutEnum> instead of TOutItems<FOutEnumItem>
+template<> struct TOutItems<bool>			: TOutItems<FOutBoolItem>		{ TOutItems(TOutItems<FOutBoolItem> Super) : TOutItems<FOutBoolItem>(Super) {} };
+template<> struct TOutItems<FOutEnum>		: TOutItems<FOutEnumItem>		{ TOutItems(TOutItems<FOutEnumItem> Super) : TOutItems<FOutEnumItem>(Super) {} };
+template<> struct TOutItems<FOutEnumerator> : TOutItems<FOutEnumeratorItem> { TOutItems(TOutItems<FOutEnumeratorItem> Super) : TOutItems<FOutEnumeratorItem>(Super) {} };
+template<> struct TOutItems<FOutFlags>		: TOutItems<FOutFlagsItem>		{ TOutItems(TOutItems<FOutFlagsItem> Super) : TOutItems<FOutFlagsItem>(Super) {} };
+template<> struct TFiller<bool>				: TFiller<FOutBoolItem>			{ using TFiller<FOutBoolItem>::TFiller; };
+template<> struct TFiller<FOutEnum>			: TFiller<FOutEnumItem>			{ using TFiller<FOutEnumItem>::TFiller; };
+template<> struct TFiller<FOutEnumerator>	: TFiller<FOutEnumeratorItem>	{ using TFiller<FOutEnumeratorItem>::TFiller; };
+template<> struct TFiller<FOutFlags>		: TFiller<FOutFlagsItem>		{ using TFiller<FOutFlagsItem>::TFiller; };
 
 ////////////////////////////////////////////////////////////////////////////////////////////////
 
@@ -515,6 +924,8 @@ struct IOps
 {
 	virtual ~IOps() {}
 	
+	inline	FOpaqueEnum			NameFlatEnum(FTypename Enum) { return { .Id = NameEnum(Enum), .bFlagMode = false }; }
+	inline	FOpaqueEnum			NameFlagEnum(FTypename Enum) { return { .Id = NameEnum(Enum), .bFlagMode = true  }; }
 	virtual FEnumId				NameEnum(FTypename Enum) = 0;
 	virtual FEnumeratorNames	NameEnumerators(FEnumId Enum, std::initializer_list<const char*> Enumerators) = 0;
 	inline  FEnumeratorNames	NameEnumerators(FTypename Enum, std::initializer_list<const char*> Enumerators)		{ return NameEnumerators(NameEnum(Enum), Enumerators); }
@@ -620,14 +1031,14 @@ struct FHistoryDeleter
 	PLAINPROPS_API void operator()(FHistory* Ptr) const;
 };
 using FHistoryPtr = TUniquePtr<FHistory, FHistoryDeleter>;
-PLAINPROPS_API FHistoryPtr MakeHistory(IDeclarations& Latest, FLiteralIndexerBase& Ids);
-
-////////////////////////////////////////////////////////////////////////////////////////////////
 
+PLAINPROPS_API FHistoryPtr			MakeHistory(IDeclarations& Latest, FLiteralIndexerBase& Ids);
 PLAINPROPS_API FMemberId			ResolveName(const FLooseMember& Member, const FHistory& History);
 PLAINPROPS_API FNameId				ResolveEnumerator(const FLooseMember& FlatEnum, const FHistory& History);
 PLAINPROPS_API TArray<FNameId>		ResolveFlags(const FLooseMember& FlagEnum, const FHistory& History);
 
+////////////////////////////////////////////////////////////////////////////////////////////////
+
 struct FUpgradedInstance
 {
 	FDeclId										Type;
diff --git a/Engine/Plugins/Experimental/PlainProps/Source/Tests/PlainPropsUpgradeTests.cpp b/Engine/Plugins/Experimental/PlainProps/Source/Tests/PlainPropsUpgradeTests.cpp
--- a/Engine/Plugins/Experimental/PlainProps/Source/Tests/PlainPropsUpgradeTests.cpp
+++ b/Engine/Plugins/Experimental/PlainProps/Source/Tests/PlainPropsUpgradeTests.cpp
@@ -22,9 +22,12 @@
 #include "PlainPropsWrite.h"
 #include "Containers/AnsiString.h"
 #include "Containers/Map.h"
+#include "Logging/StructuredLog.h"
 #include "Templates/UnrealTemplate.h"
 #include "Tests/TestHarnessAdapter.h"
 
+DEFINE_LOG_CATEGORY_STATIC(LogPlainPropsUpgradeTests, Log, All);
+
 namespace PlainProps::Upgrade::Test
 {
 
@@ -105,6 +108,60 @@ private:
 
 ////////////////////////////////////////////////////////////////////////////////////////////////
 
+FAnsiStringView	ToStr(EMismatch Mismatch)
+{
+	switch (Mismatch)
+	{
+	case EMismatch::Inners:		return "inners of";
+	case EMismatch::Type:		return "type of";
+	case EMismatch::Duplicate:	return "duplicated";
+	case EMismatch::Surplus:	return "surplus";
+	case EMismatch::Missing:	return "missing";
+	}
+	return "UNKNOWN MISMATCH";
+}
+
+static void CheckMismatches(TConstArrayView<FUpgradedInstance> Instances, FDebugPrinter Debug)
+{
+	for (FUpgradedInstance Instance : Instances)
+	{
+		if (Instance.Fixed)
+		{
+			continue;
+		}
+
+		TMap<FParameterId, FLooseType> Missing;
+		TUtf8StringBuilder<256> Msg;
+		for (FLooseMember Member : Instance.Loose)
+		{
+			EMismatch Mismatch = Member.GetValue().Meta.Mismatch;
+			FParameterId Name = Member.GetName();
+			if (Mismatch == EMismatch::No)
+			{
+				if (!Member)
+				{
+					Missing.Emplace(Name, Member.GetType()); 
+				}
+				continue;
+			}
+
+			if (Msg.Len())
+			{
+				Msg << ", ";
+			}
+			
+			Msg << ToStr(Mismatch) << " '" << Debug.Print(Instance.Type, Name) << "' " ;
+			if (const FLooseType* Expected = Missing.Find(Name))
+			{
+				Msg << "from " << Debug.Print(*Expected) << " to ";
+			}
+			Msg << Debug.Print(Member.GetType());
+		}
+
+		UE_LOGFMT(LogPlainPropsUpgradeTests, Fatal, "Missing upgrade(s) in '{Struct}' that transforms {Msg}", Debug.Print(Instance.Type), Msg);
+	}
+}
+
 // FBatchWriter/FBatchSaver helper
 struct FSavedStruct
 {
@@ -163,12 +220,15 @@ class FBatchSaver : public FBatchWriterBase
 public:
 	using FBatchWriterBase::FBatchWriterBase;
 
-	template<class T>
-	void						Save(T&& Object) 
+	template<class InT>
+	void						Save(InT&& Object, InT&& Default = {}) 
 	{
-		FDualStructId Id = IndexStructDualId<FRuntimeIds, TTypename<std::remove_reference_t<T>>>();
+		using T = std::remove_reference_t<InT>;
+		FDualStructId Id = IndexStructDualId<FRuntimeIds, TTypename<T>>();
 		BindObsolete<T>(Id);
 		Roots.Emplace(FDeclId(Id), SaveStruct(&Object, Id, {{GSchemas, GCustoms}, {Scratch}}));
+		// todo: Switch to this and fix FInputDependencies setting up type info for optional in parameters ([PlainProps] --section Basics)
+		//Roots.Emplace(FDeclId(Id), SaveStructDelta(&Object, FBaseline(&Default), Id, {{GSchemas, GCustoms}, {Scratch}}));
 	}
 
 	TArray64<uint8>				Write() const { return FBatchWriterBase::Write(GDeltaDeclarations); }
@@ -315,6 +375,7 @@ static void Run(FHistory& History, void (*SaveOld)(FBatchSaver&), void (*TestNew
 		CHECK(!Upgrader.MatchSchemas(Declarations));
 		FScratchAllocator Scratch;
 		TArray<FUpgradedInstance> NewRoots = Upgrader.Upgrade(OldBatch.GetRoots(), History, Scratch);
+		CheckMismatches(NewRoots, FDebugPrinter(History.Ids, History.LooseIds));
 
 		// Write new schemas and instances
 		FBatchSaver NewBatch(Scratch);
@@ -340,7 +401,7 @@ struct FUpgradedBatch
 	TConstArrayView<FUpgradedInstance>		Roots;
 	FLooseIndexers&							LooseIds;
 
-	FAnsiString ResolveName(FDeclId Id, FLooseMember Member) const { return GIds.ResolveName(LooseIds.Get(Id).Resolve(Member.GetValue().Meta.Name).Get().Id); }
+	FAnsiString ResolveName(FDeclId Id, FLooseMember Member) const { return GIds.ResolveName(LooseIds.Get(Id).Resolve(Member.GetName()).Get().Id); }
 };
 
 static void RunUntyped(void (*WriteUntyped)(FBatchWriter&), void (*ChronicleUntyped)(FTestUpgrades), void (*TestUntyped)(FUpgradedBatch))
@@ -451,7 +512,7 @@ PP_REFLECT_ENUM(PlainProps::Upgrade::Test, EFlatten, A, B, AB);
 
 ////////////////////////////////////////////////////////////////////////////////////////////////
 
-enum class EFlag : uint8 {A = 8, B = 32};
+enum class EFlag : uint8 {None = 0, A = 8, B = 32, AB = 40};
 PP_REFLECT_ENUM(PlainProps::Upgrade::Test, EFlag, A, B);
 
 struct FPoint { float X, Y; };
@@ -515,7 +576,7 @@ void Record(TChroniclesOf<FBasic> Out)
 		// Executed last if only FlagA or FlagB was sparsely saved
 		Out.Add({~"FlagA"_B1, ~"FlagB"_B1}, {"Flags"_Enum(AB)}, [](FIns In, FOuts Out, FContext)
 			{
-				FOutEnumFlags& Flags = Out[0].EnumFlags();
+				FOutFlags Flags = Out[0].EnumFlags();
 				Flags[/* A */ 0] = In[/* FlagA */ 0].OptBool();
 				Flags[/* B */ 1] = In[/* FlagB */ 1].OptBool();
 			});
@@ -531,7 +592,7 @@ void Record(TChroniclesOf<FBasic> Out)
 		FStructMemberNames XY = Out.NameMembers({"PlainProps::Upgrade::Test", "FPoint"}, {"X"_F32, "Y"_F32});
 		Out.Add({ ~"PointX"_F32, ~"PointY"_F32 }, { "Point"_Struct(XY) }, [](FIns In, FOuts Out, FContext)
 			{
-				TOuts<2> Point = Out[0].StructMembers<2>();
+				FOutMembers Point = Out[0].StructMembers();
 				Point[/* X */ 0].F32() = In[/* PointX */ 0].OptF32().Get(0.f);
 				Point[/* Y */ 1].F32() = In[/* PointY */ 1].OptF32().Get(0.f);
 			});
@@ -545,7 +606,7 @@ void Record(TChroniclesOf<FBasic> Out)
 		FStructMemberNames ABCD = Out.NameMembers(GuidId, {"A"_U32, "B"_U32, "C"_U32, "D"_U32});
 		Out.Add({ "GuidRenamed"_Struct(ABCD) }, { "GuidAB"_U64, "GuidCD"_U64 }, [](FIns In, FOuts Out, FContext)
 			{
-				TIns<4> ABCD = In[0].StructMembers<4>();
+				FInMembers ABCD = In[0].StructMembers();
 
 				uint64 A = ABCD[0].U32();
 				uint64 B = ABCD[1].U32();
@@ -619,7 +680,7 @@ void Record(TChroniclesOf<FBasic> Out)
 //		Out.AddTypechecked({~"FlagA"_B1, ~"FlagB"_B1}, {"Flags"_Enum(AB)}, 
 //			[](TOptional<bool> FlagA, TOptional<bool> FlagB, FOuts Out, FContext)
 //			{
-//				FOutEnumFlags& Flags = Out[0].EnumFlags();
+//				FOutFlags& Flags = Out[0].EnumFlags();
 //				Flags[/* A */ 0] = FlagA;
 //				Flags[/* B */ 1] = FlagB;
 //			});
@@ -666,13 +727,18 @@ void Record(TChroniclesOf<FBasic> Out)
 struct FStr
 {
 	TArray<char8_t> Data;
+	bool operator==(const FStr&) const = default;
 };
+PP_REFLECT_STRUCT(PlainProps::Upgrade::Test, FStr, void, Data);
+
 
 struct FPoints
 {
 	TArray<float> Xs;
 	TArray<float> Ys;
+	bool operator==(const FPoints&) const = default;
 };
+PP_REFLECT_STRUCT(PlainProps::Upgrade::Test, FPoints, void, Xs, Ys); 
 
 // Test convert FRangesA into FRangesB and vice versa
 struct FRangesA
@@ -682,8 +748,9 @@ struct FRangesA
 	TArray<FPoint>			Points; // SoA <-> AoS
 	TArray<TArray<FPoint>>	Series; // Join <-> Split
 	TArray<char>			Unterminated; // Add <-> Remove null terminator
+	TArray<EFlag>			Sorted; // Reverse opaque enums
 };
-PP_REFLECT_STRUCT(PlainProps::Upgrade::Test, FRangesA, void, Leaves, String, Points, Series, Unterminated); 
+PP_REFLECT_STRUCT(PlainProps::Upgrade::Test, FRangesA, void, Leaves, String, Points, Series, Unterminated, Sorted); 
 
 struct FRangesB
 {
@@ -691,64 +758,146 @@ struct FRangesB
 	FStr					String;
 	FPoints					Points;
 	TArray<FPoint>			Series;
-	TArray<int32>			SeriesLengths;
+	TArray<int32>			SeriesStarts; // Series indices
 	TArray<char>			Terminated;
+	TArray<EFlag>			Reversed;
+	bool operator==(const FRangesB&) const = default;
 };
-PP_REFLECT_STRUCT(PlainProps::Upgrade::Test, FRangesB, void, Leaves, String, Points, Series, SeriesLengths, Terminated);
+PP_REFLECT_STRUCT(PlainProps::Upgrade::Test, FRangesB, void, Leaves, String, Points, Series, SeriesStarts, Terminated, Reversed);
+
+void Record(TChroniclesOf<FRangesB> Out)
+{
+	FChronicler RangesA(Out, "PlainProps::Upgrade::Test", "FRangesA"); // ~ registers chronicle once we've collected all member names
+	
+	RangesA.In<R0>([](IOps& Out)
+	{
+		FEnumFlagNames AB = Out.NameFlags({"PlainProps::Upgrade::Test", "EFlag"}, {"A", "B"});
+		Out.Add({"Leaves"_U8s}, {"Leaves"_Enums(AB)}, [](FIns In, FOuts Out, FContext Ctx)
+		{
+			TInItems<uint8> Ints = In[0].U8s().Items();
+			TFiller<FOutFlags> Enums = Out[0].EnumsFlags(2).New(Ints.Num(), Ctx);
+			for (uint8 Int : Ints)
+			{
+				Enums << uint64((/* A */ Int & 1) | (/* B */ Int & 2));
+			}
+		});
+		
+		FStructMemberNames Data = Out.NameMembers({"PlainProps::Upgrade::Test", "FStr"}, {"Data"_Str});
+		Out.Add({"String"_Str}, {"String"_Struct(Data)}, [](FIns In, FOuts Out, FContext Ctx)
+		{
+			Out[/* String */ 0].StructMembers()[/* Data */ 0].Str() = In[/* String */ 0].Str();
+		});
+		
+		FStructMemberNames XY = Out.NameMembers({"PlainProps::Upgrade::Test", "FPoint"}, {"X"_F32, "Y"_F32});
+		FStructMemberNames XsYs = Out.NameMembers({"PlainProps::Upgrade::Test", "FPoints"}, {"Xs"_F32s, "Ys"_F32s});
+		Out.Add({"Points"_Structs(XY)}, {"Points"_Struct(XsYs)}, [](FIns In, FOuts Out, FContext Ctx)
+		{
+			TInItems<FInMembers> InPoints = In[0].StructsMembers(2).Items();
+			FOutMembers OutPoints = Out[0].StructMembers();
+			TFiller<float> Xs = OutPoints[0].F32s().New(InPoints.Num(), Ctx);
+			TFiller<float> Ys = OutPoints[1].F32s().New(InPoints.Num(), Ctx);
+			for (FInMembers XY : InPoints)
+			{
+				Xs << XY[0].F32();
+				Ys << XY[1].F32();
+			}
+		});
+
+		FStructId Point = Out.NameStruct({"PlainProps::Upgrade::Test", "FPoint"});
+		Out.Add({"Series"_Structs(Point).Range()}, {"Series"_Structs(Point), "SeriesStarts"_S32s}, [](FIns In, FOuts Out, FContext Ctx)
+		{
+			TInItems<FInStruct, 2> InSeries = In[0].Range<FInStruct, 2>().Items();
+			TFiller<int32> OutStarts = Out[1].S32s().New(InSeries.Num(), Ctx);
+			int32 Start = 0;
+			for (TInItems<FInStruct> InSerie : InSeries)
+			{
+				OutStarts << Start;
+				Start += InSerie.Num();
+			}
+
+			TFiller<FOutStruct> OutSeries = Out[0].Structs().New(Start, Ctx);
+			for (TInItems<FInStruct> InSerie : InSeries)
+			{
+				OutSeries << InSerie;
+			}
+		});
+
+		Out.Add({"Unterminated"_Str}, {"Terminated"_Str}, [](FIns In, FOuts Out, FContext Ctx)
+		{
+			TInItems<char8_t> Unterminated = In[0].Str().Items();
+			TFiller<char8_t> Terminated = Out[0].Str().New(Unterminated.Num() + 1, Ctx);
+			Terminated << Unterminated << '\0';
+		});
+
+		FOpaqueEnum Flag = Out.NameFlagEnum({"PlainProps::Upgrade::Test", "EFlag"});
+		Out.Add({"Sorted"_Enums(Flag)}, {"Reversed"_Enums(Flag)}, [](FIns In, FOuts Out, FContext Ctx)
+		{
+			TInItems<FInEnum> Sorted = In[0].Enums().Items();
+			TOutItems<FOutEnum> Reversed = Out[0].Enums().New(Sorted.Num(), Ctx);
+			for (uint64 Idx = Sorted.Num(); FInEnum Enum : Sorted)
+			{
+				Reversed[--Idx] = Enum;
+			}
+		});
+
+	});
+	RangesA.Rename<R0>({"PlainProps::Upgrade::Test", "FRangesB"});
+}
+		//Out.Transform("Leaves"_U8, "Leaves"_Enum(Flag)), [](uint8 In, FOutEnumFlag Out)
+		//{
+		//	Out[/* A */ 0] = In & 1;
+		//	Out[/* B */ 1] = In & 2;
+		//});
+
+		//TEnum<2, R0> Flag(Out, "PlainProps::Upgrade::Test", "EFlag", {"A", "B"});
+		//TStruct<2, R0> Point(Out, "PlainProps::Upgrade::Test", "FPoint", {"X", "Y"});
+		//TStruct<2, R0> Points(Out, "PlainProps::Upgrade::Test", "FPoints", {"Xs", "Ys"});
+		//TStruct<2, R0> Str(Out, "PlainProps::Upgrade::Test", "FStr", {"Data"});
+
+
+	//		RangesA.Transform("Leaves"_U8, "Leaves"_Enum(Flag)),
+	//	[](uint8 In, FOutEnum Out)
+	//	{
+	//		Out.EmitIf(/* A */ 0, !!(In & 1));
+	//		Out.EmitIf(/* B */ 1, !!(In & 2));
+	//	});
+	//RangesA.Structify("String"_UTF8, "String"_Struct(Str));  // overload without lambda
+	//RangesA.Structify("Points"_Range.Of(Point), "Points"_Struct(Points),
+	//	[](FInStructs Points, FOutFloats Xs, FOutFloats Ys)
+	//	{
+	//		FFloatWriter XIt = Xs.Emit(In.Num());
+	//		FFloatWriter YIt = Ys.Emit(In.Num());
+	//		for (FInStruct Point : Points)
+	//		{
+	//			XIt.Add(Point[/* X */ 0].AsRequired<float>());
+	//			YIt.Add(Point[/* Y */ 1].AsRequired<float>());
+	//		}
+	//	});
+
+	//RangesA.Upgrade("Series"_Range.Nest().Of(Point), "Series"_Range.Of(Point), "SeriesLengths"_Range.Of<int32>(),
+	//	[](FInNestedStructs SeriesOfSeries, FOutStructs Series, FOutI32s SeriesLengths)
+	//	{
+	//		FStructWriter SeriesIt = Series.Emit(SumNums(SeriesOfSeries));
+	//		FI32Writer LengthsIt = SeriesLengths.Emit(SeriesOfSeries.Num());
+	//		for (FInStructRange Points : SeriesOfSeries)
+	//		{
+	//			LengthsIt.Add(IntCastChecked<int32>(Points.Num()));
+	//			for (FInStruct Point : Points)
+	//			{
+	//				SeriesIt.Add(Point);
+	//			}
+	//		}
+	//	});
+
+	//RangesA.Upgrade<FInUtf8s>("Unterminated", "Terminated"),
+	//	[](FInUtf8s Unterminated, FOutUtf8s Terminated)
+	//	{
+	//		FUtf8s OutData = Terminated.Emit(Unterminated.Num() + 1);
+	//		OutData.EmitSlice(0, Unterminated);
+	//		OutData.EmitItem(Unterminated.Num, '\0');
+	//	});
+	//});
 
-//
-//// Sugar
-//template<> PlainProps::Upgrade::Chronicle<FRangesA>(FHistory& Out)
-//{
-//	TEnum<2, R0> Flag(Out, "", "EFlag", {"A", "B"});
-//	TStruct<2, R0> Point(Out, "", "FPoint", {"X", "Y"});
-//	TStruct<2, R0> Points(Out, "", "FPoints", {"Xs", "Ys"});
-//	TStruct<2, R0> Str(Out, "", "FStr", {"Data"});
-//
-//	FChronicle RangesA(Out, "", "FRangesA"); // ~ registers chronicle once we've collected all member names
-//	RangesA.Transform("Leaves"_U8, "Leaves"_Enum(Flag)),
-//		[](uint8 In, FOutEnum Out)
-//		{
-//			Out.EmitIf(/* A */ 0, !!(In & 1));
-//			Out.EmitIf(/* B */ 1, !!(In & 2));
-//		});
-//	RangesA.Structify("String"_UTF8, "String"_Struct(Str));  // overload without lambda
-//	RangesA.Structify("Points"_Range.Of(Point), "Points"_Struct(Points),
-//		[](FInStructs Points, FOutFloats Xs, FOutFloats Ys)
-//		{
-//			FFloatWriter XIt = Xs.Emit(In.Num());
-//			FFloatWriter YIt = Ys.Emit(In.Num());
-//			for (FInStruct Point : Points)
-//			{
-//				XIt.Add(Point[/* X */ 0].AsRequired<float>());
-//				YIt.Add(Point[/* Y */ 1].AsRequired<float>());
-//			}
-//		});
-//
-//	RangesA.Upgrade("Series"_Range.Nest().Of(Point), "Series"_Range.Of(Point), "SeriesLengths"_Range.Of<int32>(),
-//		[](FInNestedStructs SeriesOfSeries, FOutStructs Series, FOutI32s SeriesLengths)
-//		{
-//			FStructWriter SeriesIt = Series.Emit(SumNums(SeriesOfSeries));
-//			FI32Writer LengthsIt = SeriesLengths.Emit(SeriesOfSeries.Num());
-//			for (FInStructRange Points : SeriesOfSeries)
-//			{
-//				LengthsIt.Add(IntCastChecked<int32>(Points.Num()));
-//				for (FInStruct Point : Points)
-//				{
-//					SeriesIt.Add(Point);
-//				}
-//			}
-//		});
-//
-//	RangesA.Upgrade<FInUtf8s>("Unterminated", "Terminated"),
-//		[](FInUtf8s Unterminated, FOutUtf8s Terminated)
-//		{
-//			FUtf8s OutData = Terminated.Emit(Unterminated.Num() + 1);
-//			OutData.EmitSlice(0, Unterminated);
-//			OutData.EmitItem(Unterminated.Num, '\0');
-//		});
-//}
-//
 //// Sugar (move down)
 //template<> PlainProps::Upgrade::Chronicle<FRangesB>(FUpgrades& Out)
 //{
@@ -857,6 +1006,27 @@ TEST_CASE_NAMED(FPlainPropsUpgradeTest, "System::PlainProps::Upgrade", "[PlainPr
 
 	SECTION("Ranges")
 	{
+		TScopedStructBinding<FPoint, FDefaultRuntime> Point;
+		TScopedStructBinding<FPoints, FDefaultRuntime> Points;
+		TScopedStructBinding<FStr, FDefaultRuntime> Str;
+		TScopedEnumDeclaration<EFlag, EEnumMode::Flag, FDefaultRuntime> Flag;
+
+		FHistoryPtr History = MakeHistory(GDeltaDeclarations, GIds);
+		TScopedStructUpgradeBinding<FRangesB> RangesB(*History);
+		Run(*History,
+			[](FBatchSaver& Batch)
+			{
+				Batch.Save(FRangesA{.Leaves = {0,1,2,3}, .String = {'h','i'}, .Points = {{1,5}, {2,6}, {3,7}}, 
+									.Series = {{}, {{1,2}, {3,4}, {5,6}}, {}, {{7,8}}, {}},
+									.Unterminated = {'l','o'}, .Sorted = {EFlag::A, EFlag::B, EFlag::AB} });
+			}, 
+			[](FBatchLoader& Batch)
+			{
+				FRangesB B = Batch.Load<FRangesB>(0);
+				CHECK(B == FRangesB{.Leaves = {EFlag::None, EFlag::A, EFlag::B, EFlag::AB}, .String = {{'h', 'i',}}, .Points = {{1,2,3},{5,6,7}},
+									.Series = {{1,2}, {3,4}, {5,6}, {7,8}}, .SeriesStarts = {0,0,3,3,4}, 
+									.Terminated = {'l','o','\0'}, .Reversed = {EFlag::AB, EFlag::B, EFlag::A}});
+			});
 	}
 	
 	SECTION("Versions")
