PlainProps upgrader implement type-checking for missing optional input and switch tests to sparse

#rnx
#Synced-CL 50765653

diff --git a/Engine/Plugins/Experimental/PlainProps/Source/Private/PlainPropsInternalUpgrade.h b/Engine/Plugins/Experimental/PlainProps/Source/Private/PlainPropsInternalUpgrade.h
--- a/Engine/Plugins/Experimental/PlainProps/Source/Private/PlainPropsInternalUpgrade.h
+++ b/Engine/Plugins/Experimental/PlainProps/Source/Private/PlainPropsInternalUpgrade.h
@@ -212,6 +212,7 @@ struct FOp
 	FMiniSlice	Outputs;
 	uint8		NumOptionalInputs;
 	uint8		NumRequiredInputs;
+	uint16		OptionalInputIdx;
 
 	int32		NumInputs() const { return NumOptionalInputs + NumRequiredInputs; }
 };
@@ -258,6 +259,7 @@ struct FVersionUpgrades
 	uint16						NumOps;
 	uint16						NumDecompositions;	// Struct / enum inputs with named members / enumerators / flags
 	uint16						NumRecompositions;	// Enum outputs with named enumerators / flags
+	uint16						NumOptionalInputs;	// Total optional inputs, enables type-checking missing optional inputs
 	FOptionalVersion			NextVersion;		// Next version for this OR renamed struct
 	FMiniSlice					NameMatches[0];		// #NumNames
 	
@@ -267,7 +269,8 @@ struct FVersionUpgrades
 	const FMemberMatch*			GetMatches() const			{ return AlignPtr<FMemberMatch>(NameMatches + NumNames); }
 	const FOp*					GetOps() const				{ return AlignPtr<FOp>(GetMatches() + NumMatches); }
 	const uint32*				GetDecompOffsets() const	{ return AlignPtr<uint32>(GetOps() + NumOps); }
-	const uint16*				GetRecompJumps() const		{ return AlignPtr<uint16>(GetDecompOffsets() + NumDecompositions); }
+	const FLooseParameter*		GetOptionalInputs()	const	{ return AlignPtr<FLooseParameter>(GetDecompOffsets() + NumDecompositions); }
+	const uint16*				GetRecompJumps() const		{ return AlignPtr<uint16>(GetOptionalInputs() + NumOptionalInputs); }
 	const FOutParameter*		GetOutputs() const			{ return &At<FOutParameter>(OutputsOffset); }
 };
 
diff --git a/Engine/Plugins/Experimental/PlainProps/Source/Private/PlainPropsUpgrade.cpp b/Engine/Plugins/Experimental/PlainProps/Source/Private/PlainPropsUpgrade.cpp
--- a/Engine/Plugins/Experimental/PlainProps/Source/Private/PlainPropsUpgrade.cpp
+++ b/Engine/Plugins/Experimental/PlainProps/Source/Private/PlainPropsUpgrade.cpp
@@ -161,6 +161,7 @@ static FString TlsPrint(FLooseType T) { return TlsDebug->Print(T); }
 
 struct FFriend
 {
+	static FLooseMember&			Edit(FIn& I)						{ return I; }
 	static const FLooseMember&		Read(const FOut& O)					{ return O; }
 	static TArrayView<FLooseMember>	Edit(TArrayView<FOut> Os)			{ return MakeArrayView(static_cast<FLooseMember*>(Os.GetData()), Os.Num()); }
 
@@ -1775,6 +1776,30 @@ static int64 DebugVersionNo(const FStructUpgrades& Struct, const FVersionUpgrade
 	return Struct.FirstVersion + (Pos - Struct.Versions);
 }
 
+inline bool CompareAll(const FLooseMember& Actual, FLooseParameter Expected)
+{
+	return Actual.GetName() == Expected.Name && Actual.GetType().CompareAll(Expected.Type);
+}
+
+// Enables type-checking FIn::OptXyz() functions
+static void MakeMissingInputs(TArrayView<FIn> OutInputs, TConstArrayView<FLooseParameter> OptionalInputs)
+{
+	int32 Idx = 0;
+	for (FIn& Input : OutInputs)
+	{
+		if (FLooseMember& In = FFriend::Edit(Input))
+		{
+			check(!In.GetType().OptionalParameter || CompareAll(In, OptionalInputs[Idx]));
+			Idx += In.GetType().OptionalParameter;
+		}
+		else
+		{
+			In = MakeMember(OptionalInputs[Idx++]);
+		}
+	}
+	check(Idx == OptionalInputs.Num());
+}
+
 void FInstanceUpgrader::Flush()
 {
 	check(Current);
@@ -1790,8 +1815,9 @@ void FInstanceUpgrader::Flush()
 				checkf(Deps->MissingRequired == 0, TEXT("Required input missing for op %d when upgrading version %d of %s"),
 													OpIdx, DebugVersionNo(Struct, Current), *Ctx.Debug.Print(Struct.Id));
 				check(Deps->MissingOptional > 0);
-
 				FOp Op =  Current->GetOps()[OpIdx];
+				check(Deps->MissingOptional < Op.NumOptionalInputs + (Op.NumRequiredInputs > 0));
+				MakeMissingInputs(/* out */ MakeArrayView(Deps->Inputs, Op.NumInputs()), MakeArrayView(Current->GetOptionalInputs() + Op.OptionalInputIdx, Op.NumOptionalInputs));
 				Run(Op, Deps->Inputs, /* Recipient, inputs already given to this */ 0);
 			}
 			++OpIdx;
@@ -3397,6 +3423,7 @@ struct FOps final : IOps
 	TInnerSlices<FEnumeratorId>					InnerEnumerators;
 	FOptionalDeclId								Redirect;
 	int32										NumInputs = 0;
+	int32										NumOptionalInputs = 0;
 	int32										NumOutputsFlat = 0;
 	int32										NumDecompositions = 0;
 	int32										NumStructDecompositions = 0;
@@ -3495,12 +3522,13 @@ struct FOps final : IOps
 	{
 		check(IsOpen());
 
-		uint8 NumOptionalInputs = IntCastChecked<uint8>(Algo::CountIf(Ins, [](FParameter P){ return P.bOptional; }));
-		uint8 NumRequiredInputs = IntCastChecked<uint8>(Ins.Num() - NumOptionalInputs);
-		Ops.Emplace(Func, NumOptionalInputs, NumRequiredInputs, IntCastChecked<uint8>(Outs.Num()));
+		uint8 NumOptionalIns = IntCastChecked<uint8>(Algo::CountIf(Ins, [](FParameter P){ return P.bOptional; }));
+		uint8 NumRequiredIns = IntCastChecked<uint8>(Ins.Num() - NumOptionalIns);
+		Ops.Emplace(Func, NumOptionalIns, NumRequiredIns, IntCastChecked<uint8>(Outs.Num()));
 		IndexParameters(Ins);
 		CountDecompositions(MakeArrayView(Parameters).Right(Ins.Num()));
 		NumInputs += Ins.Num();
+		NumOptionalInputs += NumOptionalIns;
 		IndexParameters(Outs);
 		NumOutputsFlat += CountFlat(MakeArrayView(Parameters).Right(Outs.Num()));
 	}
@@ -3691,10 +3719,12 @@ static void CompileOps(FVersionUpgrades& Out, const FOps& In)
 	FOutParameter* Outputs = const_cast<FOutParameter*>(Out.GetOutputs());
 	FOp* OutOps = const_cast<FOp*>(Out.GetOps());
 	uint32* DecompOffsets = const_cast<uint32*>(Out.GetDecompOffsets());
+	FLooseParameter* OptionalInputs = const_cast<FLooseParameter*>(Out.GetOptionalInputs());
 	uint16* RecompJumps = const_cast<uint16*>(Out.GetRecompJumps());
 	FEnumeratorId* RecompIdIt = reinterpret_cast<FEnumeratorId*>(RecompJumps + Out.NumRecompositions);
 	uint8* DecompDataIt = Align(reinterpret_cast<uint8*>(Outputs + In.NumOutputsFlat), alignof(FStructDecomposition));
 	uint16 DecompIdx = 0;
+	uint16 OptionalInputIdx = 0;
 	uint32 OutputIdx = 0;
 	TConstArrayView<FMidOp> InOps = In.Ops;
 	for (int32 OpIdx = 0; OpIdx < InOps.Num(); ++OpIdx)
@@ -3706,8 +3736,6 @@ static void CompileOps(FVersionUpgrades& Out, const FOps& In)
 		uint8 InputIdx = 0;
 		for (FMidParam Input : PopFront(Params, Op.NumOptionalInputs + Op.NumRequiredInputs))
 		{
-			FOptionalMemberId Foo = In.OwnerIndexer.Resolve(Input.Name);
-
 			// Generate input decomposition match
 			if (Input.InnerSlice)
 			{
@@ -3721,6 +3749,12 @@ static void CompileOps(FVersionUpgrades& Out, const FOps& In)
 				++DecompIdx;
 			}
 
+			// Copy optional input metadata
+			if (Input.Type.OptionalParameter)
+			{
+				OptionalInputs[OptionalInputIdx++] = {Input.Name, Input.Type};
+			}
+
 			// Copy normal input match
 			FMatchIdx Idx = { .Op = static_cast<uint16>(OpIdx) };
 			AllMatches[Input.Name.Idx].Add({ .Type = Input.Type, .Idx = Idx, .InputIdx = InputIdx, .Recipient = 0 });
@@ -3747,11 +3781,13 @@ static void CompileOps(FVersionUpgrades& Out, const FOps& In)
 		OutOps[OpIdx] = {	.Function = Op.Function,
 							.Outputs = { OutputIdx, IntCastChecked<uint8>(NumOutputsFlat) },
 							.NumOptionalInputs = Op.NumOptionalInputs,
-							.NumRequiredInputs = Op.NumRequiredInputs };
+							.NumRequiredInputs = Op.NumRequiredInputs,
+							.OptionalInputIdx = IntCastChecked<uint16>(OptionalInputIdx - Op.NumOptionalInputs) };
 
 		OutputIdx += NumOutputsFlat;
 	}
 	check(OutputIdx == In.NumOutputsFlat);
+	check(OptionalInputIdx == Out.NumOptionalInputs);
 	check(RecompIdIt == reinterpret_cast<FEnumeratorId*>(RecompJumps + Out.NumRecompositions) + In.InnerEnumerators.Items.Num());
 	check(Out.GetOutputs() == AlignPtr<FOutParameter>(RecompIdIt));
 
@@ -3775,6 +3811,7 @@ static void CompileVersion(FVersionUpgrades& Out, const FOps& In)
 	Out.NumOps = IntCastChecked<uint16>(In.Ops.Num());
 	Out.NumDecompositions = IntCastChecked<uint16>(In.NumDecompositions);
 	Out.NumRecompositions = IntCastChecked<uint16>(In.NumRecompositions());
+	Out.NumOptionalInputs = IntCastChecked<uint16>(In.NumOptionalInputs);
 	const FEnumeratorId* RecompIds = AlignPtr<FEnumeratorId>(Out.GetRecompJumps() + In.NumRecompositions());
 	const FOutParameter* Outputs = AlignPtr<FOutParameter>(RecompIds + In.NumRecompositionIds());
 	Out.OutputsOffset = IntCastChecked<uint32>((const uint8*)Outputs - (uint8*)&Out);
@@ -3789,6 +3826,7 @@ static uint32 CalculateVersionUpgradesSize(const FOps& Ops)
 	Out = Align(Out + Ops.NumMatches() * sizeof(FMemberMatch), alignof(FMemberMatch));
 	Out = Align(Out + Ops.Ops.Num() * sizeof(FOp), alignof(FOp));
 	Out = Align(Out + Ops.NumDecompositions * sizeof(uint32), alignof(uint32));
+	Out = Align(Out + Ops.NumOptionalInputs * sizeof(FLooseParameter), alignof(FLooseParameter));
 	Out = Align(Out + Ops.NumRecompositions() * sizeof(uint16), alignof(uint16));
 	Out = Align(Out + Ops.NumRecompositionIds() * sizeof(FEnumeratorId), alignof(FEnumeratorId));
 	Out = Align(Out + Ops.NumOutputsFlat * sizeof(FOutParameter), alignof(FOutParameter));
diff --git a/Engine/Plugins/Experimental/PlainProps/Source/Public/PlainPropsUpgrade.h b/Engine/Plugins/Experimental/PlainProps/Source/Public/PlainPropsUpgrade.h
--- a/Engine/Plugins/Experimental/PlainProps/Source/Public/PlainPropsUpgrade.h
+++ b/Engine/Plugins/Experimental/PlainProps/Source/Public/PlainPropsUpgrade.h
@@ -494,6 +494,8 @@ private:
 		check(IsOptional());
 		return Value ? TOptional<T>(ArithmeticCast<T>(Value.Data.Arithmetic)) : NullOpt;
 	}
+
+	friend FFriend;
 };
 
 // Opaque struct value, assignable to a type-matching FOutStruct
diff --git a/Engine/Plugins/Experimental/PlainProps/Source/Tests/PlainPropsUpgradeTests.cpp b/Engine/Plugins/Experimental/PlainProps/Source/Tests/PlainPropsUpgradeTests.cpp
--- a/Engine/Plugins/Experimental/PlainProps/Source/Tests/PlainPropsUpgradeTests.cpp
+++ b/Engine/Plugins/Experimental/PlainProps/Source/Tests/PlainPropsUpgradeTests.cpp
@@ -226,9 +226,7 @@ public:
 		using T = std::remove_reference_t<InT>;
 		FDualStructId Id = IndexStructDualId<FRuntimeIds, TTypename<T>>();
 		BindObsolete<T>(Id);
-		Roots.Emplace(FDeclId(Id), SaveStruct(&Object, Id, {{GSchemas, GCustoms}, {Scratch}}));
-		// todo: Switch to this and fix FInputDependencies setting up type info for optional in parameters ([PlainProps] --section Basics)
-		//Roots.Emplace(FDeclId(Id), SaveStructDelta(&Object, FBaseline(&Default), Id, {{GSchemas, GCustoms}, {Scratch}}));
+		Roots.Emplace(FDeclId(Id), SaveStructDelta(&Object, FBaseline(&Default), Id, {{GSchemas, GCustoms}, {Scratch}}));
 	}
 
 	TArray64<uint8>				Write() const { return FBatchWriterBase::Write(GDeltaDeclarations); }
