PlainProps: When printing flag enum values use a zero enumerator instead of !UnsetFlagEnum when one exists.
Most CoreUObject flag enums have an explicity zero enumerator and it makes sense to use that friendly name when it exists.

#rnx

#Synced-CL 50856727

diff --git a/Engine/Plugins/Experimental/PlainProps/Source/Private/PlainPropsPrint.cpp b/Engine/Plugins/Experimental/PlainProps/Source/Private/PlainPropsPrint.cpp
--- a/Engine/Plugins/Experimental/PlainProps/Source/Private/PlainPropsPrint.cpp
+++ b/Engine/Plugins/Experimental/PlainProps/Source/Private/PlainPropsPrint.cpp
@@ -342,7 +342,14 @@ static void PrintEnumValue(
 
 	if (Value == 0)
 	{
-		Buffer << GLiterals.UnsetFlagEnum;
+		if (FNameId ConstantId = GetEnumConstantId<T, true>(Schema, Value); ConstantId != FNameId())
+		{
+			Ids.AppendString(Buffer, ConstantId);
+		}
+		else
+		{
+			Buffer << GLiterals.UnsetFlagEnum;
+		}
 		return;
 	}
 
diff --git a/Engine/Plugins/Experimental/PlainProps/Source/Tests/PlainPropsTests.cpp b/Engine/Plugins/Experimental/PlainProps/Source/Tests/PlainPropsTests.cpp
--- a/Engine/Plugins/Experimental/PlainProps/Source/Tests/PlainPropsTests.cpp
+++ b/Engine/Plugins/Experimental/PlainProps/Source/Tests/PlainPropsTests.cpp
@@ -999,7 +999,7 @@ TEST_CASE_NAMED(FPlainPropsReadWriteTest, "System::Core::Serialization::PlainPro
 		enum class EFlatSparse8 : int8 { A = 1, B = 2, C = 3 };
 
 		static constexpr std::initializer_list<const char*> MemberNames = 
-		{ "A2", "A0", "B0", "B4", "B5", "B7", "C3", "D34", "Max8", "Max16", "Max32", "Max64", "IF" };
+		{ "A2", "A0", "B0", "B4", "B5", "B7", "C3", "D0", "D34", "Max8", "Max16", "Max32", "Max64", "IF" };
 
 		TestSerialize([](FTestBatchBuilder& Batch, FScratchAllocator& Scratch)
 		{
@@ -1011,7 +1011,7 @@ TEST_CASE_NAMED(FPlainPropsReadWriteTest, "System::Core::Serialization::PlainPro
 			FEnumId X	= Batch.DeclareEnum("Test", "UnusedEnum2",	EEnumMode::Flag, ELeafWidth::B8,	{"U4"}, {1}); // Hole
 			FEnumId B	= Batch.DeclareEnum("Test", "FlagDense8",	EEnumMode::Flag, ELeafWidth::B8,	{"A", "B", "C"}, {1, 2, 4});
 			FEnumId C	= Batch.DeclareEnum("Test", "FlatSparse8",	EEnumMode::Flat, ELeafWidth::B8,	{"A", "B", "C"}, {1, 2, 3});
-			FEnumId D	= Batch.DeclareEnum("Test", "FlagSparse8",	EEnumMode::Flag, ELeafWidth::B8,	{"A", "B", "C"}, {2, 16, 32});
+			FEnumId D	= Batch.DeclareEnum("Test", "FlagSparse8",	EEnumMode::Flag, ELeafWidth::B8,	{"None", "A", "B", "C"}, {0, 2, 16, 32});
 			FEnumId E	= Batch.DeclareEnum("Test", "FlatLimit8",	EEnumMode::Flat, ELeafWidth::B8,	{"Min", "Max"}, {0, 0xFF});
 			FEnumId F	= Batch.DeclareEnum("Test", "FlatLimit16",	EEnumMode::Flat, ELeafWidth::B16,	{"Min", "Max"}, {0, 0xFFFF});
 			FEnumId G	= Batch.DeclareEnum("Test", "FlatLimit32",	EEnumMode::Flat, ELeafWidth::B32,	{"Min", "Max"}, {0, 0xFFFFFFFF});
@@ -1020,7 +1020,7 @@ TEST_CASE_NAMED(FPlainPropsReadWriteTest, "System::Core::Serialization::PlainPro
 			
 			FMemberSpec MemberTypes[] = {
 				Enum8(A), Enum8(A), Enum8(B), Enum8(B), Enum8(B), Enum8(B), 
-				Enum8(C), Enum8(D), Enum8(E), Enum16(F), Enum32(G), Enum64(H), Enum64(I) };
+				Enum8(C), Enum8(D), Enum8(D), Enum8(E), Enum16(F), Enum32(G), Enum64(H), Enum64(I) };
 			FDeclId ObjectId = Batch.DeclareStruct("Test", "Enums", MemberNames, MemberTypes, EMemberPresence::AllowSparse);
 
 			FMemberBuilder Members;
@@ -1031,6 +1031,7 @@ TEST_CASE_NAMED(FPlainPropsReadWriteTest, "System::Core::Serialization::PlainPro
 			Members.AddEnum(Batch.NameMember("B5"),				B, (uint8)5);
 			Members.AddEnum(Batch.NameMember("B7"),				B, (uint8)7);
 			Members.AddEnum(Batch.NameMember("C3"),				C, EFlatSparse8::C);
+			Members.AddEnum(Batch.NameMember("D0"),				D, (uint8)0);
 			Members.AddEnum(Batch.NameMember("D34"),			D, (uint8)34);
 			Members.AddEnum(Batch.NameMember("Max8"),			E, (uint8)0xFF);
 			Members.AddEnum(Batch.NameMember("Max16"),			F, (uint16)0xFFFF);
@@ -1059,6 +1060,7 @@ TEST_CASE_NAMED(FPlainPropsReadWriteTest, "System::Core::Serialization::PlainPro
 			CHECK(It1.GrabLeaf().AsUnderlyingValue<uint8>() == 5);
 			CHECK(It1.GrabLeaf().AsUnderlyingValue<uint8>() == 7);
 			CHECK(It1.GrabLeaf().As<EFlatSparse8>() == EFlatSparse8::C);
+			CHECK(It1.GrabLeaf().AsUnderlyingValue<uint8>() == 0);
 			CHECK(It1.GrabLeaf().AsUnderlyingValue<uint8>() == 34);
 			CHECK(It1.GrabLeaf().AsUnderlyingValue<uint8>() == 0xFF);
 			CHECK(It1.GrabLeaf().AsUnderlyingValue<uint16>() == 0xFFFF);
@@ -1075,6 +1077,7 @@ TEST_CASE_NAMED(FPlainPropsReadWriteTest, "System::Core::Serialization::PlainPro
 			CHECK(GetEnumName(It2.GrabLeaf()) == "FlagDense8");
 			CHECK(GetEnumName(It2.GrabLeaf()) == "FlatSparse8");
 			CHECK(GetEnumName(It2.GrabLeaf()) == "FlagSparse8");
+			CHECK(GetEnumName(It2.GrabLeaf()) == "FlagSparse8");
 			CHECK(GetEnumName(It2.GrabLeaf()) == "FlatLimit8");
 			CHECK(GetEnumName(It2.GrabLeaf()) == "FlatLimit16");
 			CHECK(GetEnumName(It2.GrabLeaf()) == "FlatLimit32");
@@ -1090,6 +1093,7 @@ TEST_CASE_NAMED(FPlainPropsReadWriteTest, "System::Core::Serialization::PlainPro
 			(void)It3.GrabLeaf();
 			const FEnumSchema& FlatSparse8 = ResolveEnumSchema(Batch, It3.GrabLeaf().Enum.Get());
 			const FEnumSchema& FlagSparse8 = ResolveEnumSchema(Batch, It3.GrabLeaf().Enum.Get());
+			(void)It3.GrabLeaf();
 			const FEnumSchema& FlatLimit8  = ResolveEnumSchema(Batch, It3.GrabLeaf().Enum.Get());
 			const FEnumSchema& FlatLimit16 = ResolveEnumSchema(Batch, It3.GrabLeaf().Enum.Get());
 			const FEnumSchema& FlatLimit32 = ResolveEnumSchema(Batch, It3.GrabLeaf().Enum.Get());
@@ -1109,7 +1113,7 @@ TEST_CASE_NAMED(FPlainPropsReadWriteTest, "System::Core::Serialization::PlainPro
 			CHECK(EqualEnumNames(MakeConstArrayView(FlatDense8.Footer, FlatDense8.Num), MakeConstArrayView({"A", "C"})));
 			CHECK(EqualEnumNames(MakeConstArrayView(FlagDense8.Footer, FlagDense8.Num), MakeConstArrayView({"A", "B", "C"})));
 			CHECK(EqualEnumNames(MakeConstArrayView(FlatSparse8.Footer, FlatSparse8.Num), MakeConstArrayView({"C"})));
-			CHECK(EqualEnumNames(MakeConstArrayView(FlagSparse8.Footer, FlagSparse8.Num), MakeConstArrayView({"A", "C"})));
+			CHECK(EqualEnumNames(MakeConstArrayView(FlagSparse8.Footer, FlagSparse8.Num), MakeConstArrayView({"None", "A", "C"})));
 			CHECK(Names[FlatLimit8.Footer[0]] == "Max");
 			CHECK(Names[FlatLimit16.Footer[0]] == "Max");
 			CHECK(Names[FlatLimit32.Footer[0]] == "Max");
@@ -1120,7 +1124,7 @@ TEST_CASE_NAMED(FPlainPropsReadWriteTest, "System::Core::Serialization::PlainPro
 			CHECK(EqualItems(GetConstants<uint8>(FlatDense8), MakeConstArrayView({0, 2})));
 			CHECK(EqualItems(GetConstants<uint8>(FlagDense8), TConstArrayView<uint8>()));
 			CHECK(EqualItems(GetConstants<EFlatSparse8>(FlatSparse8), MakeConstArrayView({EFlatSparse8::C})));
-			CHECK(EqualItems(GetConstants<uint8>(FlagSparse8), MakeConstArrayView({2, 32})));
+			CHECK(EqualItems(GetConstants<uint8>(FlagSparse8), MakeConstArrayView({0, 2, 32})));
 			CHECK(GetConstants<uint8>(FlatLimit8)[0] == 0xFF);
 			CHECK(GetConstants<uint16>(FlatLimit16)[0] == 0xFFFF);
 			CHECK(GetConstants<uint32>(FlatLimit32)[0] == 0xFFFFFFFF);
