PlainProps upgrade add flat enum loosener with range support and bitcast optimizations for enum range loosening

#rnx
#Synced-CL 50765653

diff --git a/Engine/Plugins/Experimental/PlainProps/Source/Private/PlainPropsInternalUpgrade.h b/Engine/Plugins/Experimental/PlainProps/Source/Private/PlainPropsInternalUpgrade.h
--- a/Engine/Plugins/Experimental/PlainProps/Source/Private/PlainPropsInternalUpgrade.h
+++ b/Engine/Plugins/Experimental/PlainProps/Source/Private/PlainPropsInternalUpgrade.h
@@ -300,9 +300,14 @@ class FEnumLoosener
 	FEnumLoosener(const void* Loosener, uint64 Flags) : Ptr(Flags | reinterpret_cast<uint64>(Loosener)) {}
 public:
 	FEnumLoosener() {}
-	FEnumLoosener(bool bNeeded, const FFlatEnumLoosener& Flat UE_LIFETIMEBOUND) : FEnumLoosener(&Flat, bNeeded * NeededBit) {}
-	FEnumLoosener(bool bNeeded, const FFlagEnumLoosener& Flag UE_LIFETIMEBOUND) : FEnumLoosener(&Flag, bNeeded * NeededBit | FlagBit) {}
-	
+	FEnumLoosener(FEnumLoosener&& O) : Ptr(O.Ptr) { O.Ptr = 0; }
+	FEnumLoosener(const FEnumLoosener&) = delete;
+	FEnumLoosener(bool bNeeded, /* own */ FFlatEnumLoosener* Flat UE_LIFETIMEBOUND) : FEnumLoosener(Flat, bNeeded * NeededBit) {}
+	FEnumLoosener(bool bNeeded, /* ref */ const FFlagEnumLoosener& Flag UE_LIFETIMEBOUND) : FEnumLoosener(&Flag, bNeeded * NeededBit | FlagBit) {}
+	~FEnumLoosener();
+	FEnumLoosener& operator=(FEnumLoosener&& O) { Swap(Ptr, O.Ptr); return *this; }
+	FEnumLoosener& operator=(const FEnumLoosener&) = delete;
+
 	bool						Exists() const { return !!(Ptr); }
 	bool						Needed() const { return !!(Ptr & NeededBit); }
 	bool						IsFlag() const { return !!(Ptr & FlagBit); }
@@ -420,7 +425,7 @@ public:
 struct FBatchContext
 {
 	// todo: add loose transformations for redirected and rename to source transformations
-	const FTransformations		Transformations;
+	const FTransformations&		Transformations;
 	FScratchAllocator&			Scratch;
 	const FIdIndexerBase&		Ids;
 	FLooseIndexers&				LooseIds; // RenameMembers can index new names
diff --git a/Engine/Plugins/Experimental/PlainProps/Source/Private/PlainPropsUpgrade.cpp b/Engine/Plugins/Experimental/PlainProps/Source/Private/PlainPropsUpgrade.cpp
--- a/Engine/Plugins/Experimental/PlainProps/Source/Private/PlainPropsUpgrade.cpp
+++ b/Engine/Plugins/Experimental/PlainProps/Source/Private/PlainPropsUpgrade.cpp
@@ -602,6 +602,15 @@ FRangeData FOutRange::NewArithmetics(uint64 Num, FContext Ctx, uint32 SizeOf) co
 
 ////////////////////////////////////////////////////////////////////////////////////////////////
 
+FEnumLoosener::~FEnumLoosener()
+{
+	if (Exists() && !IsFlag())
+	{
+		TMap<uint64, FEnumeratorId> Empty;
+		Swap(reinterpret_cast<FFlatEnumLoosener*>(Ptr & PtrMask)->ConstantNames, Empty);
+	}
+}
+
 FEnumLooseners::FEnumLooseners(TConstArrayView<FEnumTransformation> Xforms, uint32 InNumStructs)
 : Data(Xforms.GetData())
 , NumStructs(InNumStructs)
@@ -720,7 +729,7 @@ FLooseValue FLooseMemberReader::Grab()
 
 	if (Leaf.Enum)
 	{
-		FEnumLoosener Loosener = Looseners[Leaf.Enum.Get()];
+		const FEnumLoosener& Loosener = Looseners[Leaf.Enum.Get()];
 		if (Loosener.Needed())
 		{
 			FEnumeratorData Data = Next.Type.InnermostType.Enum.FlagMode 
@@ -1079,7 +1088,7 @@ static FStructRange LoosenStructRangeOnly(FLooseNestedRange Range, FSchemaId Sch
 }
 
 template<typename T, typename LoosenerType>
-static void LoosenEnums(FEnumLooseData* Dst, TConstArrayView64<T> Src, const LoosenerType& Loosener)
+void LoosenEnums(FEnumLooseData* Dst, TConstArrayView64<T> Src, const LoosenerType& Loosener)
 {
 	for (T Constant : Src)
 	{
@@ -1088,15 +1097,29 @@ static void LoosenEnums(FEnumLooseData* Dst, TConstArrayView64<T> Src, const Loo
 }
 
 template<typename T>
-static void LoosenEnums(FEnumLooseData* Dst, const T* Src, uint64 Num, const FEnumLoosener& Loosener)
+void BitCastEnums(FEnumLooseData* Dst, TConstArrayView64<T> Src)
+{
+	for (T Constant : Src)
+	{
+		*Dst++ = { .Enumerator = BitCast<FEnumeratorData>(uint64(Constant)) };
+	}
+}
+
+template<typename T>
+void LoosenEnums(FEnumLooseData* Dst, const T* Src, uint64 Num, const FEnumLoosener& Loosener)
 {
+	// @see LoosenerRequired()
+	if (!Loosener.Exists())
+	{
+		return BitCastEnums<T>(Dst, MakeArrayView64(Src, Num));
+	}
+
 	return Loosener.IsFlag()	? LoosenEnums<T>(Dst, MakeArrayView64(Src, Num), Loosener.AsFlag())
 								: LoosenEnums<T>(Dst, MakeArrayView64(Src, Num), Loosener.AsFlat());
 }
 
 static FEnumRange LoosenEnums(ELeafWidth SourceWidth, FLooseNestedRange Range, const FEnumLoosener& Loosener, FScratchAllocator& Scratch)
 {
-	check(Loosener.Exists());
 	const void* Src = Range.Data.Enums.Source;
 	FEnumLooseData* Dst = Scratch.AllocateArray<FEnumLooseData>(Range.Num);
 	switch (SourceWidth)
@@ -1424,7 +1447,7 @@ FLooseValue FInstanceUpgrader::Decompose(FEnumMember In, const FEnumDecompositio
 	FEnumeratorData Enumerator = {};
 	if (In.State == ELooseState::Source)
 	{
-		FEnumLoosener Loosener = Ctx.Transformations[In.Schema.Get()].Source;
+		const FEnumLoosener& Loosener = Ctx.Transformations[In.Schema.Get()].Source;
 		Enumerator = In.FlagMode	? Loosen(Loosener.AsFlag(), In.Data.Source)
 									: Loosen(Loosener.AsFlat(), In.Data.Source);
 	}
@@ -2432,6 +2455,19 @@ static void AddInnerIdOutputs(TSet<FDeclId>& OutStructs, TSet<FEnumId>& OutEnums
 	}
 }
 
+
+static bool LoosenerRequired(const FFlatEnumLoosener& Loosener)
+{
+	for (TPair<uint64, FEnumeratorId> Pair : Loosener.ConstantNames)
+	{
+		if (Pair.Key != Pair.Value.Idx)
+		{
+			return true;
+		}
+	}
+	return false;
+}
+
 static bool LoosenerRequired(const FFlagEnumLoosener& Loosener, uint32 Num)
 {
 	for (uint32 Idx = 0; Idx < Num; ++Idx)
@@ -2444,6 +2480,30 @@ static bool LoosenerRequired(const FFlagEnumLoosener& Loosener, uint32 Num)
 	return false;
 }
 
+template<typename T>
+static FEnumLoosener CreateFlatLoosener(const FEnumSchema& Schema, bool bNeeded, FEnumeratorIndexer& Ids, FScratchAllocator& Scratch)
+{
+	FFlatEnumLoosener Out;
+	Out.ConstantNames.Reserve(Schema.Num);
+	const FNameId* Names = Schema.Footer;
+	if (const T* Constants = Schema.ExplicitConstants ? GetConstants<T>(Schema).GetData() : nullptr)
+	{
+		for (uint32 Idx = 0, Num = Schema.Num; Idx < Num; ++Idx)
+		{
+			Out.ConstantNames.Emplace(Constants[Idx], Ids.Index(Names[Idx]));
+		}
+	}
+	else
+	{
+		for (uint32 Idx = 0, Num = Schema.Num; Idx < Num; ++Idx)
+		{
+			Out.ConstantNames.Emplace(Idx, Ids.Index(Names[Idx]));
+		}
+	}
+
+	return LoosenerRequired(Out) ? FEnumLoosener(bNeeded, new (Scratch) FFlatEnumLoosener(MoveTemp(Out))) : FEnumLoosener();
+}
+
 template<typename T>
 static FEnumLoosener CreateFlagLoosener(const FEnumSchema& Schema, bool bNeeded, FEnumeratorIndexer& Ids, FScratchAllocator& Scratch)
 {
@@ -2468,13 +2528,6 @@ static FEnumLoosener CreateFlagLoosener(const FEnumSchema& Schema, bool bNeeded,
 	return LoosenerRequired(Out, Schema.Num) ? FEnumLoosener(bNeeded, *new (Scratch) FFlagEnumLoosener(Out)) : FEnumLoosener();
 }
 
-template<typename T>
-static FEnumLoosener CreateFlatLoosener(const FEnumSchema& Schema, bool bNeeded, FEnumeratorIndexer& Ids, FScratchAllocator& Scratch)
-{
-	unimplemented();
-	return FEnumLoosener();
-}
-
 static FEnumLoosener CreateEnumLoosener(const FEnumSchema& Schema, bool bNeeded, FEnumeratorIndexer& Ids, FScratchAllocator& Scratch)
 {
 	switch (Schema.Width)
diff --git a/Engine/Plugins/Experimental/PlainProps/Source/Tests/PlainPropsUpgradeTests.cpp b/Engine/Plugins/Experimental/PlainProps/Source/Tests/PlainPropsUpgradeTests.cpp
--- a/Engine/Plugins/Experimental/PlainProps/Source/Tests/PlainPropsUpgradeTests.cpp
+++ b/Engine/Plugins/Experimental/PlainProps/Source/Tests/PlainPropsUpgradeTests.cpp
@@ -722,6 +722,9 @@ void Record(TChroniclesOf<FBasic> Out)
 
 ////////////////////////////////////////////////////////////////////////////////////////////////
 
+enum class EFlat : uint8 {A=3, B=9, C=0};
+PP_REFLECT_ENUM(PlainProps::Upgrade::Test, EFlat, A, B, C);
+
 struct FStr
 {
 	TArray<char8_t> Data;
@@ -729,7 +732,6 @@ struct FStr
 };
 PP_REFLECT_STRUCT(PlainProps::Upgrade::Test, FStr, void, Data);
 
-
 struct FPoints
 {
 	TArray<float> Xs;
@@ -746,7 +748,7 @@ struct FRangesA
 	TArray<FPoint>			Points; // SoA <-> AoS
 	TArray<TArray<FPoint>>	Series; // Join <-> Split
 	TArray<char>			Unterminated; // Add <-> Remove null terminator
-	TArray<EFlag>			Sorted; // Reverse opaque enums
+	TArray<EFlat>			Sorted; // Reverse opaque enums
 };
 PP_REFLECT_STRUCT(PlainProps::Upgrade::Test, FRangesA, void, Leaves, String, Points, Series, Unterminated, Sorted); 
 
@@ -758,7 +760,7 @@ struct FRangesB
 	TArray<FPoint>			Series;
 	TArray<int32>			SeriesStarts; // Series indices
 	TArray<char>			Terminated;
-	TArray<EFlag>			Reversed;
+	TArray<EFlat>			Reversed;
 	bool operator==(const FRangesB&) const = default;
 };
 PP_REFLECT_STRUCT(PlainProps::Upgrade::Test, FRangesB, void, Leaves, String, Points, Series, SeriesStarts, Terminated, Reversed);
@@ -827,8 +829,8 @@ void Record(TChroniclesOf<FRangesB> Out)
 			Terminated << Unterminated << '\0';
 		});
 
-		FOpaqueEnum Flag = Out.NameFlagEnum({"PlainProps::Upgrade::Test", "EFlag"});
-		Out.Add({"Sorted"_Enums(Flag)}, {"Reversed"_Enums(Flag)}, [](FIns In, FOuts Out, FContext Ctx)
+		FOpaqueEnum Flat = Out.NameFlatEnum({"PlainProps::Upgrade::Test", "EFlat"});
+		Out.Add({"Sorted"_Enums(Flat)}, {"Reversed"_Enums(Flat)}, [](FIns In, FOuts Out, FContext Ctx)
 		{
 			TInItems<FInEnum> Sorted = In[0].Enums().Items();
 			TOutItems<FOutEnum> Reversed = Out[0].Enums().New(Sorted.Num(), Ctx);
@@ -1007,6 +1009,7 @@ TEST_CASE_NAMED(FPlainPropsUpgradeTest, "System::PlainProps::Upgrade", "[PlainPr
 		TScopedStructBinding<FPoint, FDefaultRuntime> Point;
 		TScopedStructBinding<FPoints, FDefaultRuntime> Points;
 		TScopedStructBinding<FStr, FDefaultRuntime> Str;
+		TScopedEnumDeclaration<EFlat, EEnumMode::Flat, FDefaultRuntime> Flat;
 		TScopedEnumDeclaration<EFlag, EEnumMode::Flag, FDefaultRuntime> Flag;
 
 		FHistoryPtr History = MakeHistory(GDeltaDeclarations, GIds);
@@ -1016,14 +1019,14 @@ TEST_CASE_NAMED(FPlainPropsUpgradeTest, "System::PlainProps::Upgrade", "[PlainPr
 			{
 				Batch.Save(FRangesA{.Leaves = {0,1,2,3}, .String = {'h','i'}, .Points = {{1,5}, {2,6}, {3,7}}, 
 									.Series = {{}, {{1,2}, {3,4}, {5,6}}, {}, {{7,8}}, {}},
-									.Unterminated = {'l','o'}, .Sorted = {EFlag::A, EFlag::B, EFlag::AB} });
+									.Unterminated = {'l','o'}, .Sorted = {EFlat::A, EFlat::B, EFlat::C} });
 			}, 
 			[](FBatchLoader& Batch)
 			{
 				FRangesB B = Batch.Load<FRangesB>(0);
 				CHECK(B == FRangesB{.Leaves = {EFlag::None, EFlag::A, EFlag::B, EFlag::AB}, .String = {{'h', 'i',}}, .Points = {{1,2,3},{5,6,7}},
 									.Series = {{1,2}, {3,4}, {5,6}, {7,8}}, .SeriesStarts = {0,0,3,3,4}, 
-									.Terminated = {'l','o','\0'}, .Reversed = {EFlag::AB, EFlag::B, EFlag::A}});
+									.Terminated = {'l','o','\0'}, .Reversed = {EFlat::C, EFlat::B, EFlat::A}});
 			});
 	}
 	
