PlainProps implement decomposed enum ranges, finish bool ranges and add tests

#rnx

#Synced-CL 50905510

diff --git a/Engine/Plugins/Experimental/PlainProps/Source/Private/PlainPropsInternalUpgrade.h b/Engine/Plugins/Experimental/PlainProps/Source/Private/PlainPropsInternalUpgrade.h
--- a/Engine/Plugins/Experimental/PlainProps/Source/Private/PlainPropsInternalUpgrade.h
+++ b/Engine/Plugins/Experimental/PlainProps/Source/Private/PlainPropsInternalUpgrade.h
@@ -479,6 +479,7 @@ class FInstanceUpgrader
 	FLooseValue						Decompose(FEnumMember Input, const FEnumDecomposition& Decomp);
 	FLooseValue						Decompose(FStructMember Input, const FStructDecomposition& Decomp);
 	TArrayView<FLooseMember>		Decompose(FStructView Input, const FStructDecomposition& Decomp, FParameterId Name, EMismatch& OutMismatch);
+	FEnumIndexData					Reorder(FEnumeratorData Input, const TSet<FEnumeratorId>& Order, bool bFlagMode);
 	FEnumIndexData					Reorder(FEnumeratorId Flat, const TSet<FEnumeratorId>& Order);
 	FEnumIndexData					Reorder(FEnumeratorIdSet Flags, const TSet<FEnumeratorId>& Order);
 	void							Input(FMemberMatch Match, FLooseValue Value);
diff --git a/Engine/Plugins/Experimental/PlainProps/Source/Private/PlainPropsUpgrade.cpp b/Engine/Plugins/Experimental/PlainProps/Source/Private/PlainPropsUpgrade.cpp
--- a/Engine/Plugins/Experimental/PlainProps/Source/Private/PlainPropsUpgrade.cpp
+++ b/Engine/Plugins/Experimental/PlainProps/Source/Private/PlainPropsUpgrade.cpp
@@ -534,6 +534,12 @@ void FOutEnumItem::operator=(FInEnum In)
 	Out = In.Value;
 }
 
+void FOutEnumeratorItem::operator=(uint16 Index)
+{
+	check(Index < Max);
+	Out.Index = Index;
+}
+
 void FOutFlagItem::operator=(bool bValue)
 {
 	SetBit(Out.Bits, Idx, bValue);
@@ -1357,6 +1363,7 @@ static FRangeData LoosenItemsForDecomposition(FLooseType ItemType, FSchemaId Sch
 	}
 	else if (ItemType.IsStruct())
 	{
+		check(!ItemType.InnermostType.Struct.IsDecomposed);
 		FDeclId Id = ItemType.InnermostId.Get().AsStructDeclId();
 		const FLooseSchema& LooseSchema = ObtainSourceSchema(Ctx, Id, FStructSchemaId(SchemaId));
 		FEnumLooseners Looseners = ViewLooseners(Ctx);
@@ -1374,10 +1381,10 @@ static FRangeData LoosenItemsForDecomposition(FLooseType ItemType, FSchemaId Sch
 	else
 	{
 		check(ItemType.IsEnum());
-		unimplemented();
-		FEnumLooseData* Out = nullptr;
-
-		return {.Enums = {.Loose = Out}};
+		check(!ItemType.InnermostType.Enum.IsDecomposed);
+		FEnumSchemaId Id(SchemaId);
+		const FEnumSchema& Schema = ResolveEnumSchema(Ctx.Batch, Id);
+		return { .Enums = LoosenEnums(Schema.Width, Range, ViewLooseners(Ctx)[Id], Ctx.Scratch) };
 	}
 }
 
@@ -1456,10 +1463,8 @@ FLooseValue FInstanceUpgrader::Decompose(FEnumMember In, const FEnumDecompositio
 		Enumerator = In.Data.Loose.Enumerator;
 	}
 
-	FEnumIndexData Decomposed = In.FlagMode	? Reorder(Enumerator.Flags, Decomp.InputOrder) 
-											: Reorder(Enumerator.Flat, Decomp.InputOrder);
-	FLooseMetadata Meta = { .IsSet = 1, .Name = In.Name };
-	return { Meta, {.Enum = {.Loose = {.Decomposed = Decomposed}}} };
+	FEnumIndexData Decomposed = Reorder(Enumerator, Decomp.InputOrder, In.FlagMode);
+	return { {.IsSet = 1, .Name = In.Name}, {.Enum = {.Loose = {.Decomposed = Decomposed}}} };
 }
 
 static int32 GetInputIndex(FEnumeratorId Name, const TSet<FEnumeratorId>& Order)
@@ -1469,6 +1474,15 @@ static int32 GetInputIndex(FEnumeratorId Name, const TSet<FEnumeratorId>& Order)
 	return Idx.AsInteger();
 }
 
+FEnumIndexData FInstanceUpgrader::Reorder(FEnumeratorData In, const TSet<FEnumeratorId>& Order, bool bFlagMode)
+{
+	if (Order.IsEmpty())
+	{
+		return BitCast<FEnumIndexData>(In);
+	}
+	return bFlagMode ? Reorder(In.Flags, Order) : Reorder(In.Flat, Order); 
+}
+
 FEnumIndexData FInstanceUpgrader::Reorder(FEnumeratorId Flat, const TSet<FEnumeratorId>& Order)
 {
 	return { .Flat = { .Index = static_cast<uint16>(GetInputIndex(Flat, Order)) } };
diff --git a/Engine/Plugins/Experimental/PlainProps/Source/Public/PlainPropsUpgrade.h b/Engine/Plugins/Experimental/PlainProps/Source/Public/PlainPropsUpgrade.h
--- a/Engine/Plugins/Experimental/PlainProps/Source/Public/PlainPropsUpgrade.h
+++ b/Engine/Plugins/Experimental/PlainProps/Source/Public/PlainPropsUpgrade.h
@@ -149,6 +149,7 @@ constexpr FParameter operator ""_Char(const char* Name, std::size_t Len)	{ retur
 constexpr FParameter operator ""_C8(const char* Name, std::size_t Len)		{ return MakeParam<char8_t>(Name, Len); }
 constexpr FParameter operator ""_C16(const char* Name, std::size_t Len)		{ return MakeParam<char16_t>(Name, Len); }
 constexpr FParameter operator ""_C32(const char* Name, std::size_t Len)		{ return MakeParam<char32_t>(Name, Len); }
+constexpr FParameter operator ""_Bools(const char* Name, std::size_t Len)	{ return MakeParam<bool>(Name, Len).Range(); }
 constexpr FParameter operator ""_B1s(const char* Name, std::size_t Len)		{ return MakeParam<bool>(Name, Len).Range(); }
 constexpr FParameter operator ""_S8s(const char* Name, std::size_t Len)		{ return MakeParam<int8>(Name, Len).Range(); }
 constexpr FParameter operator ""_U8s(const char* Name, std::size_t Len)		{ return MakeParam<uint8>(Name, Len).Range(); }
@@ -410,8 +411,8 @@ struct FOutRange
 template<typename T>
 FRangeData FOutRange::NewItems(uint64 Num, FContext Ctx) const
 {
-	if constexpr (std::is_same_v<T, bool>)						return NewBools(Num, Ctx);
-	else if constexpr (Arithmetic<T>)							return NewArithmetics(Num, Ctx, sizeof(T));
+	if constexpr (Arithmetic<T> && !std::is_same_v<T, bool>)	return NewArithmetics(Num, Ctx, sizeof(T));
+	else if constexpr (std::is_same_v<T, FOutBoolItem>)			return NewBools(Num, Ctx);
 	else if constexpr (std::is_same_v<T, FOutStruct>)			return NewStructs(Num, Ctx);
 	else if constexpr (std::is_same_v<T, FOutAny>)				return NewStructs(Num, Ctx);
 	else if constexpr (std::is_same_v<T, FOutMembers>)			return NewMembers(Num, Ctx);
@@ -688,7 +689,7 @@ struct FOutFlags : protected FOutMember
 class FOutBoolItem
 {
 	uint8&					Out;
-	uint8					Mask;
+	const uint8				Mask;
 public:
 	FOutBoolItem(uint8& Byte, uint8 InMask) : Out(Byte), Mask(InMask) {}
 	PLAINPROPS_API void operator=(bool bValue);
@@ -696,17 +697,17 @@ public:
 
 class FOutEnumItem
 {
-	const FLooseType		Type;
 	FEnumeratorData&		Out;
+	const FLooseType		Type;
 public:
-	FOutEnumItem(FLooseType T, FEnumeratorData& D) : Type(T), Out(D) {}
+	FOutEnumItem(FLooseType T, FEnumeratorData& D) : Out(D), Type(T) {}
 	PLAINPROPS_API void operator=(FInEnum In);
 };
 
 class FOutEnumeratorItem
 {
 	FEnumeratorIndex&		Out;
-	uint32					Max;
+	const uint32			Max;
 public:
 	FOutEnumeratorItem(FEnumeratorIndex& Name, uint32 InMax) : Out(Name), Max(InMax) {}
 	PLAINPROPS_API void operator=(uint16 Index);
diff --git a/Engine/Plugins/Experimental/PlainProps/Source/Tests/PlainPropsUpgradeTests.cpp b/Engine/Plugins/Experimental/PlainProps/Source/Tests/PlainPropsUpgradeTests.cpp
--- a/Engine/Plugins/Experimental/PlainProps/Source/Tests/PlainPropsUpgradeTests.cpp
+++ b/Engine/Plugins/Experimental/PlainProps/Source/Tests/PlainPropsUpgradeTests.cpp
@@ -749,8 +749,10 @@ struct FRangesA
 	TArray<TArray<FPoint>>	Series; // Join <-> Split
 	TArray<char>			Unterminated; // Add <-> Remove null terminator
 	TArray<EFlat>			Sorted; // Reverse opaque enums
+	TArray<EFlat>			Enums;	// A->B, B->C, C->A
+	TArray<bool>			Bools;	// Negate
 };
-PP_REFLECT_STRUCT(PlainProps::Upgrade::Test, FRangesA, void, Leaves, String, Points, Series, Unterminated, Sorted); 
+PP_REFLECT_STRUCT(PlainProps::Upgrade::Test, FRangesA, void, Leaves, String, Points, Series, Unterminated, Sorted, Enums, Bools); 
 
 struct FRangesB
 {
@@ -761,9 +763,11 @@ struct FRangesB
 	TArray<int32>			SeriesStarts; // Series indices
 	TArray<char>			Terminated;
 	TArray<EFlat>			Reversed;
+	TArray<EFlat>			Cycled;
+	TArray<bool>			Negated;
 	bool operator==(const FRangesB&) const = default;
 };
-PP_REFLECT_STRUCT(PlainProps::Upgrade::Test, FRangesB, void, Leaves, String, Points, Series, SeriesStarts, Terminated, Reversed);
+PP_REFLECT_STRUCT(PlainProps::Upgrade::Test, FRangesB, void, Leaves, String, Points, Series, SeriesStarts, Terminated, Reversed, Cycled, Negated);
 
 void Record(TChroniclesOf<FRangesB> Out)
 {
@@ -840,6 +844,27 @@ void Record(TChroniclesOf<FRangesB> Out)
 			}
 		});
 
+		FEnumeratorNames ABC = Out.NameEnumerators({"PlainProps::Upgrade::Test", "EFlat"}, {"A","B","C"});
+		Out.Add({"Enums"_Enums(ABC)}, {"Cycled"_Enums(ABC)}, [](FIns In, FOuts Out, FContext Ctx)
+		{
+			TInItems<FInEnumerator> Enums = In[0].Enumerators(3).Items();
+			TFiller<FOutEnumerator> Cycled = Out[0].Enumerators(3).New(Enums.Num(), Ctx);
+			for (FInEnumerator Enum : Enums)
+			{
+				Cycled << (Enum.GetIndex() + 1) % 3;
+			}
+		});
+
+		Out.Add({"Bools"_Bools}, {"Negated"_Bools}, [](FIns In, FOuts Out, FContext Ctx)
+		{
+			TInItems<bool> Bools = In[0].Bools().Items();
+			TFiller<bool> Negated = Out[0].Bools().New(Bools.Num(), Ctx);
+			for (bool Bool : Bools)
+			{
+				Negated << !Bool;
+			}
+		});
+
 	});
 	RangesA.Rename<R0>({"PlainProps::Upgrade::Test", "FRangesB"});
 }
@@ -1019,14 +1044,16 @@ TEST_CASE_NAMED(FPlainPropsUpgradeTest, "System::PlainProps::Upgrade", "[PlainPr
 			{
 				Batch.Save(FRangesA{.Leaves = {0,1,2,3}, .String = {'h','i'}, .Points = {{1,5}, {2,6}, {3,7}}, 
 									.Series = {{}, {{1,2}, {3,4}, {5,6}}, {}, {{7,8}}, {}},
-									.Unterminated = {'l','o'}, .Sorted = {EFlat::A, EFlat::B, EFlat::C} });
+									.Unterminated = {'l','o'}, .Sorted = {EFlat::A, EFlat::B, EFlat::C},
+									.Enums = {EFlat::A, EFlat::C}, .Bools = {true, false, true} });
 			}, 
 			[](FBatchLoader& Batch)
 			{
 				FRangesB B = Batch.Load<FRangesB>(0);
 				CHECK(B == FRangesB{.Leaves = {EFlag::None, EFlag::A, EFlag::B, EFlag::AB}, .String = {{'h', 'i',}}, .Points = {{1,2,3},{5,6,7}},
 									.Series = {{1,2}, {3,4}, {5,6}, {7,8}}, .SeriesStarts = {0,0,3,3,4}, 
-									.Terminated = {'l','o','\0'}, .Reversed = {EFlat::C, EFlat::B, EFlat::A}});
+									.Terminated = {'l','o','\0'}, .Reversed = {EFlat::C, EFlat::B, EFlat::A},
+									.Cycled = {EFlat::B, EFlat::A}, .Negated = {false, true, false}});
 			});
 	}
 	
