PlainProps: Enum improvements.
Remove FOptionalLeafWidth and FEnumDeclarations::Specify() since UEnum has a new UEnum::EUnderlyingType member.
Tweak DeclareEnum(UEnum*) to declare the width and to mask out sign extended bits for negatvie enum constant values of width B8, B16 and B32.
Tweak DeclareNativeEnum to not sign extend negative enum constant values of width B8, B16 and B32.
Add constexpr SplitItems() with tests to StringUtil/StringTests to split an std::string_view by comma and whitespace.
Tweak Ctti macro for enums to support "infinite" enumerators by splitting up the Enumerators arrray into separate arrays for Names and Constants.
Add PP_REFLECT_FLAG_ENUM macro that validates that flag enumerators are zero or use a single bit.
Make IndexParameterName() compile for enums.
Add UeCore binding with tests for TEnumAsByte.

#rnx

#Synced-CL 51785738

diff --git a/Engine/Plugins/Experimental/PlainProps/Source/Private/PlainPropsDeclare.cpp b/Engine/Plugins/Experimental/PlainProps/Source/Private/PlainPropsDeclare.cpp
--- a/Engine/Plugins/Experimental/PlainProps/Source/Private/PlainPropsDeclare.cpp
+++ b/Engine/Plugins/Experimental/PlainProps/Source/Private/PlainPropsDeclare.cpp
@@ -87,14 +87,11 @@ static void ValidateDeclaration(const FEnumDeclaration& Enum)
 		checkf(!bDeclared, TEXT("Enumerator constant declared twice"));
 	}
 
-	if (Enum.Width)
+	uint64 OverflowMasks[4] = { ~uint64(0) << 8, ~uint64(0) << 16, ~uint64(0) << 32, 0 };
+	uint64 OverflowMask = OverflowMasks[(uint8)Enum.Width];
+	for (FEnumerator E : Enum.GetEnumerators())
 	{
-		uint64 OverflowMasks[4] = { ~uint64(0) << 8, ~uint64(0) << 16, ~uint64(0) << 32, 0 };
-		uint64 OverflowMask = OverflowMasks[(uint8)Enum.Width.Get()];
-		for (FEnumerator E : Enum.GetEnumerators())
-		{
-			checkf(!(E.Constant & OverflowMask), TEXT("Enumerator constant is larger than declared width"));
-		}
+		checkf(!(E.Constant & OverflowMask), TEXT("Enumerator constant %llx is larger than declared width"), E.Constant);
 	}
 }
 
@@ -107,7 +104,7 @@ void CopyItems(T* It, TConstArrayView<T> Items)
 	}
 }
 
-const FEnumDeclaration& FEnumDeclarations::Declare(FEnumId Id, FType Type, EEnumMode Mode, FOptionalLeafWidth Width, TConstArrayView<FEnumerator> Enumerators, EEnumAliases Policy)
+const FEnumDeclaration& FEnumDeclarations::Declare(FEnumId Id, FType Type, EEnumMode Mode, ELeafWidth Width, TConstArrayView<FEnumerator> Enumerators, EEnumAliases Policy)
 {
 	if (static_cast<int32>(Id.Idx) >= Declarations.Num())
 	{
@@ -133,20 +130,6 @@ const FEnumDeclaration& FEnumDeclarations::Declare(FEnumId Id, FType Type, EEnum
 	return *Ptr;
 }
 
-void FEnumDeclarations::Specify(FEnumId Id, ELeafWidth In)
-{
-	Check(Id);
-	FOptionalLeafWidth& Width = Declarations[Id.Idx]->Width;
-	if (Width)
-	{
-		check(Width.Get() == In);
-	}
-	else
-	{
-		Width = In;
-	}
-}
-
 #if DO_CHECK
 void FEnumDeclarations::Check(FEnumId Id) const
 {
@@ -273,4 +256,4 @@ const FStructDeclaration& IDeclarations::Get(FStructId Id) const
 	return *Out;
 }
 
-} // namespace PlainProps
\ No newline at end of file
+} // namespace PlainProps
diff --git a/Engine/Plugins/Experimental/PlainProps/Source/Private/PlainPropsUpgrade.cpp b/Engine/Plugins/Experimental/PlainProps/Source/Private/PlainPropsUpgrade.cpp
--- a/Engine/Plugins/Experimental/PlainProps/Source/Private/PlainPropsUpgrade.cpp
+++ b/Engine/Plugins/Experimental/PlainProps/Source/Private/PlainPropsUpgrade.cpp
@@ -2011,8 +2011,7 @@ static bool	MatchEnum(const FEnumDeclaration& Decl, const FEnumSchema& Schema)
 {
 	// Flag mode can differ if constants happen to match (unlikely)
 	
-	// Consider if we should require enum width declaration at this point
-	if (Decl.Width && Decl.Width.Get() != Schema.Width)
+	if (Decl.Width != Schema.Width)
 	{
 		return false;
 	}
@@ -3136,7 +3135,7 @@ FBuiltRange* FBatchRun::FastenEnums(FEnumId Id, FInnermostEnumType Type, FLooseM
 	FEnumFastener Fastener = Transformations[Id].Loose;
 	const FEnumeratorIndexer& Indexer = History.LooseIds.Get(Id);
 	const FEnumDeclaration& Declaration = History.Latest.Get(Id);
-	switch (Declaration.Width.Get())
+	switch (Declaration.Width)
 	{
 		case ELeafWidth::B8:		return FastenLooseEnums<uint8 >(Enums, Fastener, Indexer, Scratch);
 		case ELeafWidth::B16:		return FastenLooseEnums<uint16>(Enums, Fastener, Indexer, Scratch);
@@ -4106,4 +4105,4 @@ void FChronicles::CommitVersion(IOps* InOps)
 	History.CommitQueue.Emplace(Ops);
 }
 
-}} // namespace PlainProps::Upgrade
\ No newline at end of file
+}} // namespace PlainProps::Upgrade
diff --git a/Engine/Plugins/Experimental/PlainProps/Source/Public/PlainPropsBindCtti.h b/Engine/Plugins/Experimental/PlainProps/Source/Public/PlainPropsBindCtti.h
--- a/Engine/Plugins/Experimental/PlainProps/Source/Public/PlainPropsBindCtti.h
+++ b/Engine/Plugins/Experimental/PlainProps/Source/Public/PlainPropsBindCtti.h
@@ -184,7 +184,7 @@ FType IndexParameterName()
 template<class Ids, ETypename, Enumeration Enum>
 FType IndexParameterName()
 {
-	return GetEnumId<Ids, Enum>();
+	return IndexEnumName<Ids, Enum>();
 }
 
 template<class Ids, ETypename Kind, typename T>
@@ -386,7 +386,7 @@ struct TScopedEnumDeclaration
 	using Ids = typename Runtime::Ids;
 
 	FEnumId Id;
-	TScopedEnumDeclaration() : Id(DeclareNativeEnum<Enum, Ids>(Runtime::GetEnums(), Mode, LeafWidth<sizeof(Enum)>)) {}
+	TScopedEnumDeclaration() : Id(DeclareNativeEnum<Enum, Mode, Ids>(Runtime::GetEnums(), LeafWidth<sizeof(Enum)>)) {}
 	~TScopedEnumDeclaration() { Runtime::GetEnums().Drop(Id); }
 };
 
@@ -549,19 +549,20 @@ FMemberBinding BindMember(uint64 Offset, FMemberSpec& OutSpec)
 
 ////////////////////////////////////////////////////////////////////////////////////////////////
 
-template<Enumeration Enum, class Ids>
-FEnumId DeclareNativeEnum(FEnumDeclarations& Out, EEnumMode Mode, ELeafWidth Width)
+template<Enumeration Enum, EEnumMode Mode, class Ids>
+FEnumId DeclareNativeEnum(FEnumDeclarations& Out, ELeafWidth Width)
 {
 	using Ctti = CttiOf<Enum>;
-	using UnderlyingType = std::underlying_type_t<typename Ctti::Type>;
+	// Never sign extend negative enum constant values of width B8, B16 and B32
+	using UnsignedUnderlyingType = std::make_unsigned_t<std::underlying_type_t<typename Ctti::Type>>;
 
 	FType Type = IndexEnumName<Ids, Enum>();
 	FEnumId Id = Ids::IndexEnum(Type);
 	FEnumerator Enumerators[Ctti::NumEnumerators];
 	for (FEnumerator& Enumerator : Enumerators)
 	{
-		Enumerator.Name = Ids::IndexName(Ctti::Enumerators[&Enumerator - Enumerators].Name);
-		Enumerator.Constant = static_cast<uint64>(static_cast<UnderlyingType>(Ctti::Enumerators[&Enumerator - Enumerators].Constant));
+		Enumerator.Name = Ids::IndexName(ToAnsiView(Ctti::Names[&Enumerator - Enumerators]));
+		Enumerator.Constant = static_cast<uint64>(static_cast<UnsignedUnderlyingType>(Ctti::Constants[&Enumerator - Enumerators]));
 	}
 	Out.Declare(Id, Type, Mode, Width, Enumerators, EEnumAliases::Fail);
 
diff --git a/Engine/Plugins/Experimental/PlainProps/Source/Public/PlainPropsCtti.h b/Engine/Plugins/Experimental/PlainProps/Source/Public/PlainPropsCtti.h
--- a/Engine/Plugins/Experimental/PlainProps/Source/Public/PlainPropsCtti.h
+++ b/Engine/Plugins/Experimental/PlainProps/Source/Public/PlainPropsCtti.h
@@ -2,7 +2,9 @@
 
 #pragma once
 
+#include "PlainPropsStringUtil.h"
 #include "HAL/Platform.h"
+#include "HAL/PlatformMath.h"
 #include <tuple>
 
 namespace PlainProps 
@@ -11,6 +13,24 @@ namespace PlainProps
 template<typename T> struct TCttiOf { using Type = decltype(CttiOfPtr((T*)nullptr)); };
 template<typename T> using CttiOf = typename TCttiOf<T>::Type;
 
+namespace Private
+{
+template<class Ctti>
+constexpr bool ValidateEnumFlags()
+{
+	using EnumType = Ctti::Type;
+	using UnsignedUnderlyingType = std::make_unsigned_t<std::underlying_type_t<EnumType>>;
+	for (EnumType Constant : Ctti::Constants)
+	{
+		if (FPlatformMath::CountBits64(static_cast<UnsignedUnderlyingType>(Constant)) > 1)
+		{
+			return false;
+		}
+	}
+	return true;
+}
+} // namespace Private
+
 //////////////////////////////////////////////////////////////////////////
 // Private PP_REFLECT_XYZ helpers
 
@@ -30,15 +50,27 @@ template<typename T> using CttiOf = typename TCttiOf<T>::Type;
 #define _PP_REFLECT_10(N, MACRO, NS, T, M, ...)			MACRO(N-10, NS, T, M) _PP_REFLECT_9(N, MACRO, NS, T, __VA_ARGS__)
 #define _PP_REFLECT_THINGS(N, MACRO, ...)				_PP_EXPAND(_PP_REFLECT_##N(N, MACRO, __VA_ARGS__))
 
-#define _PP_REFLECT_ENUMERATOR(N, NS, T, C)				{ #C, Type::C }, 
-#define _PP_REFLECT_ENUM_INNER(N, NS, T, ES)			struct T##_Ctti { inline static constexpr char Name[] = #T; using Type = NS :: T; static constexpr int NumEnumerators = N; static constexpr struct { const char* Name; Type Constant; } Enumerators[] = { ES }; inline static constexpr char Namespace[] = #NS;  }; T##_Ctti CttiOfPtr(NS :: T*);
-#define _PP_REFLECT_ENUM(N, NS, T, ...)					_PP_REFLECT_ENUM_INNER(N, NS, T, _PP_EXPAND(_PP_REFLECT_THINGS(N, _PP_REFLECT_ENUMERATOR, NS, T, __VA_ARGS__)))					
+#define _PP_REFLECT_ENUM_INNER(NS, T, S, ...)			struct T##_Ctti \
+	{ \
+	inline static constexpr char Namespace[] = #NS; \
+	inline static constexpr char Name[] = #T; \
+	using Type = NS :: T; \
+	using enum NS :: T; \
+	inline static constexpr std::string_view AllNames = S; \
+	inline static constexpr auto Names = ::PlainProps::Private::ParseArgs<AllNames>; \
+	static constexpr NS :: T Constants[] = { __VA_ARGS__ }; \
+	static constexpr uint32 NumEnumerators = sizeof(Constants)/sizeof(NS :: T); \
+	static_assert(NumEnumerators == static_cast<uint32>(Names.size())); \
+	}; \
+	T##_Ctti CttiOfPtr(NS :: T*);
+
 #define _PP_REFLECT_STRUCT(N, NS, T, S, ...)			PP_REFLECT_STRUCT_ONLY(N, NS, T, S)				_PP_EXPAND(_PP_REFLECT_THINGS(N, PP_REFLECT_MEMBER, NS, T, __VA_ARGS__))
 #define _PP_REFLECT_STRUCT_TEMPLATE(N, NS, T, S, ...)	PP_REFLECT_STRUCT_TEMPLATE_ONLY(N, NS, T, S)	_PP_EXPAND(_PP_REFLECT_THINGS(N, PP_REFLECT_TEMPLATE_MEMBER, NS, T, __VA_ARGS__))
 
 //////////////////////////////////////////////////////////////////////////
 
-#define PP_REFLECT_ENUM(NS, T, ...)						_PP_REFLECT_ENUM(_PP_NUM_ARGS(__VA_ARGS__), NS, T, __VA_ARGS__)
+#define PP_REFLECT_ENUM(NS, T, ...)						_PP_REFLECT_ENUM_INNER(NS, T, #__VA_ARGS__, __VA_ARGS__)
+#define PP_REFLECT_FLAG_ENUM(NS, T, ...)				_PP_REFLECT_ENUM_INNER(NS, T, #__VA_ARGS__, __VA_ARGS__); static_assert(::PlainProps::Private::ValidateEnumFlags<T##_Ctti>(), "Flag enum " #T " must use one bit per enumerator")
 #define PP_REFLECT_STRUCT(NS, T, S, ...)				_PP_REFLECT_STRUCT(_PP_NUM_ARGS(__VA_ARGS__), NS, T, S, __VA_ARGS__)
 #define PP_REFLECT_STRUCT_TEMPLATE(NS, T, S, ...)		_PP_REFLECT_STRUCT_TEMPLATE(_PP_NUM_ARGS(__VA_ARGS__), NS, T, S, __VA_ARGS__)
 #define PP_NAME_STRUCT(NS, T)							struct T##_Ctti { inline static constexpr char Name[] = #T; using Type = NS :: T; inline static constexpr char Namespace[] = #NS; }; T##_Ctti CttiOfPtr(T*);
diff --git a/Engine/Plugins/Experimental/PlainProps/Source/Public/PlainPropsDeclare.h b/Engine/Plugins/Experimental/PlainProps/Source/Public/PlainPropsDeclare.h
--- a/Engine/Plugins/Experimental/PlainProps/Source/Public/PlainPropsDeclare.h
+++ b/Engine/Plugins/Experimental/PlainProps/Source/Public/PlainPropsDeclare.h
@@ -14,17 +14,6 @@ namespace PlainProps
 {
 struct FStructSpec;
 
-class FOptionalLeafWidth
-{
-	uint8 Value = 0xFF;
-public:
-	FOptionalLeafWidth() = default;
-	FOptionalLeafWidth(ELeafWidth Width) : Value(static_cast<uint8>(Width)) {}
-	//void							Set(ELeafWidth Width)		{ Value = static_cast<uint8>(Width);}
-	ELeafWidth						Get() const					{ check(*this); return static_cast<ELeafWidth>(Value); }
-	explicit operator				bool() const				{ return Value != 0xFF; }
-};
-
 ////////////////////////////////////////////////////////////////////////////////////////////////
 
 struct FEnumerator
@@ -41,7 +30,7 @@ struct FEnumDeclaration
 {
 	FType							Type;			// Could be removed
 	EEnumMode						Mode;
-	FOptionalLeafWidth				Width;
+	ELeafWidth						Width;
 	uint16							NumEnumerators;
 	FEnumerator						Enumerators[0];	// Constants must be unique, no aliases allowed
 
@@ -60,8 +49,7 @@ public:
 	UE_NONCOPYABLE(FEnumDeclarations);
 	explicit FEnumDeclarations(FDebugIds In) : Debug(In) {}
 
-	PLAINPROPS_API const FEnumDeclaration&	Declare(FEnumId Id, FType Type, EEnumMode Mode, FOptionalLeafWidth Width, TConstArrayView<FEnumerator> Enumerators, EEnumAliases Policy);
-	PLAINPROPS_API void						Specify(FEnumId Id, ELeafWidth Width);
+	PLAINPROPS_API const FEnumDeclaration&	Declare(FEnumId Id, FType Type, EEnumMode Mode, ELeafWidth Width, TConstArrayView<FEnumerator> Enumerators, EEnumAliases Policy);
 	void									Drop(FEnumId Id)				{ Check(Id); Declarations[Id.Idx].Reset(); }
 	const FEnumDeclaration&					Get(FEnumId Id) const			{ Check(Id); return *Declarations[Id.Idx].Get(); }
 	const FEnumDeclaration*					Find(FEnumId Id) const			{ return Id.Idx < (uint32)Declarations.Num() ? Declarations[Id.Idx].Get() : nullptr; }
diff --git a/Engine/Plugins/Experimental/PlainProps/Source/Public/PlainPropsStringUtil.h b/Engine/Plugins/Experimental/PlainProps/Source/Public/PlainPropsStringUtil.h
--- a/Engine/Plugins/Experimental/PlainProps/Source/Public/PlainPropsStringUtil.h
+++ b/Engine/Plugins/Experimental/PlainProps/Source/Public/PlainPropsStringUtil.h
@@ -5,6 +5,7 @@
 #include <string_view>
 #include <array>
 #include "HAL/Platform.h" // uint32 etc
+#include "Misc/AsciiSet.h"
 
 namespace PlainProps::Private
 {
@@ -56,6 +57,46 @@ struct THexString
 	inline static constexpr std::string_view Value{Chars + 16 - Len, Len};
 };
 
+inline constexpr std::string_view TrimWhitespace(std::string_view Str)
+{
+	constexpr FAsciiSet Whitespace(" \v\f\t\r\n");
+	while (!Str.empty() && Whitespace.Contains(Str.front())) Str.remove_prefix(1);
+	while (!Str.empty() && Whitespace.Contains(Str.back())) Str.remove_suffix(1);
+	return Str;
+}
+
+// Returns number of delimiter separated tokens, includes empty tokens.
+// Returns zero for empty or whitespace only Str.
+inline constexpr std::size_t CountTokens(std::string_view Str, char Delimiter)
+{
+	std::size_t Num = 1;
+	for (char Char : Str) Num += (Char == Delimiter);
+	return Num > 1 ? Num : !TrimWhitespace(Str).empty();
+}
+
+// Returns array of delimiter separated and whitespace stripped tokens, includes empty tokens.
+// Returns empty array for empty or whitespace only Str.
+template <std::size_t N>
+inline constexpr std::array<std::string_view, N> ParseTokensImpl(std::string_view Str, char Delimiter)
+{
+	std::array<std::string_view, N> Out{};
+	std::size_t Idx = 0;
+	while (Idx < N)
+	{
+		std::size_t Pos = Str.find(Delimiter);
+		Out[Idx++] = TrimWhitespace(Str.substr(0, Pos));
+		if (Pos == std::string_view::npos) break;
+		Str.remove_prefix(Pos + 1);
+	}
+	return Out;
+}
+
+// Returns array of comma separated and whitespace stripped tokens, includes empty tokens.
+// Returns empty array for empty or whitespace only Str.
+// Can be used for a string view based on #__VA__ARGS.
+template <const std::string_view& Str>
+inline constexpr auto ParseArgs = ParseTokensImpl<Private::CountTokens(Str, ',')>(Str, ',');
+
 } // namespace PlainProps::Private
 
 namespace PlainProps
@@ -67,4 +108,4 @@ inline constexpr std::string_view Concat = Private::TConcat<Strs...>::Value;
 template <uint64 N>
 inline constexpr std::string_view HexString = Private::THexString<N>::Value;
 
-} // namespace PlainProps
\ No newline at end of file
+} // namespace PlainProps
diff --git a/Engine/Plugins/Experimental/PlainProps/Source/Public/PlainPropsUeCoreBindings.h b/Engine/Plugins/Experimental/PlainProps/Source/Public/PlainPropsUeCoreBindings.h
--- a/Engine/Plugins/Experimental/PlainProps/Source/Public/PlainPropsUeCoreBindings.h
+++ b/Engine/Plugins/Experimental/PlainProps/Source/Public/PlainPropsUeCoreBindings.h
@@ -634,6 +634,42 @@ struct TMapBinding : public TSetBinding<TPair<K, V>, KeyFuncs, SetAllocator>
 
 //////////////////////////////////////////////////////////////////////////
 
+template <Enumeration E>
+struct TEnumAsByteBinding : ICustomBinding
+{
+	using Type = TEnumAsByte<E>;
+
+	const FMemberId				MemberIds[1];
+	const FEnumId				EnumId;
+
+	template<class Ids>
+	TEnumAsByteBinding(TCustomSpecifier<Ids, 1>& Spec)
+	: MemberIds{Ids::IndexMember("Value")}
+	, EnumId(GetEnumId<Ids, E>())
+	{
+		Spec.Members[0] = Specify<E>(EnumId);
+	}
+
+	inline void Save(FMemberBuilder& Dst, const TEnumAsByte<E>& Src, const TEnumAsByte<E>* Default, const FSaveContext& Ctx) const
+	{
+		Dst.AddEnum(MemberIds[0], EnumId, Src.GetValue());
+	}
+		
+	inline void Load(TEnumAsByte<E>& Dst, FStructLoadView Src, ECustomLoadMethod Method) const
+	{
+		FMemberLoader Members(Src);
+		Dst = Members.GrabLeaf().As<E>();
+	}
+
+	template<typename ContextType>
+	bool Diff(TEnumAsByte<E> A, TEnumAsByte<E> B, ContextType& Ctx) const
+	{
+		return A != B;
+	}
+};
+
+//////////////////////////////////////////////////////////////////////////
+
 //TODO: Consider macroifying parts of this, e.g PP_CUSTOM_BIND(PLAINPROPS_API, FTransform, Transform, Translate, Rotate, Scale)
 struct FTransformBinding : ICustomBinding
 {
@@ -1181,6 +1217,13 @@ struct TTypename<TPair<K,V>>
 	using Parameters = std::tuple<K, V>;
 };
 
+template <Enumeration E>
+struct TTypename<TEnumAsByte<E>>
+{
+	inline static constexpr std::string_view DeclName = "EnumAsByte";
+	using Parameters = std::tuple<E>;
+};
+
 template <typename... Ts>
 struct TTypename<TVariant<Ts...>>
 {
@@ -1301,6 +1344,14 @@ struct TCustomBind<TPair<K, V>>
 	using Type = UE::TPairBinding<K, V>;
 };
 
+template<Enumeration E>
+struct TOccupancyOf<TEnumAsByte<E>> : FRequireAll {};
+
+template<Enumeration E> struct TCustomBind<TEnumAsByte<E>>
+{
+	using Type = UE::TEnumAsByteBinding<E>;
+};
+
 template<typename... Ts>
 struct TCustomBind<TVariant<Ts...>>
 {
diff --git a/Engine/Plugins/Experimental/PlainProps/Source/Tests/PlainPropsCttiTests.cpp b/Engine/Plugins/Experimental/PlainProps/Source/Tests/PlainPropsCttiTests.cpp
--- a/Engine/Plugins/Experimental/PlainProps/Source/Tests/PlainPropsCttiTests.cpp
+++ b/Engine/Plugins/Experimental/PlainProps/Source/Tests/PlainPropsCttiTests.cpp
@@ -7,31 +7,142 @@
 namespace PlainProps::Test
 {
 
-enum class E1 : uint8_t { A, B };
+enum class F1 : uint8_t { None, A = 2, B = 1 << 7 };
+
+enum U1 : uint8_t { U1_A, U1_B };
+
+enum class E1 : uint16_t { A, B };
+
+enum class E2 : int16_t
+{
+	C0, C1, C2, C3, C4, C5, C6, C7, C8, C9, C10, C11, C12, C13, C14, C15, C16, C17, C18, C19,
+	C20, C21, C22, C23, C24, C25, C26, C27, C28, C29, C30, C31, C32, C33, C34, C35, C36, C37, C38, C39,
+	C40, C41, C42, C43, C44, C45, C46, C47, C48, C49, C50, C51, C52, C53, C54, C55, C56, C57, C58, C59,
+	C60, C61, C62, C63, C64, C65, C66, C67, C68, C69, C70, C71, C72, C73, C74, C75, C76, C77, C78, C79,
+	C80, C81, C82, C83, C84, C85, C86, C87, C88, C89, C90, C91, C92, C93, C94, C95, C96, C97, C98, C99,
+};
+
 
 namespace Actual
-{	
+{
+	PP_REFLECT_FLAG_ENUM(PlainProps::Test, F1, None, A, B);
+	PP_REFLECT_ENUM(PlainProps::Test, U1, U1_A, U1_B);
 	PP_REFLECT_ENUM(PlainProps::Test, E1, A, B);
+	PP_REFLECT_ENUM(PlainProps::Test, E2,
+		C0, C1, C2, C3, C4, C5, C6, C7, C8, C9, C10, C11, C12, C13, C14, C15, C16, C17, C18, C19,
+		C20, C21, C22, C23, C24, C25, C26, C27, C28, C29, C30, C31, C32, C33, C34, C35, C36, C37, C38, C39,
+		C40, C41, C42, C43, C44, C45, C46, C47, C48, C49, C50, C51, C52, C53, C54, C55, C56, C57, C58, C59,
+		C60, C61, C62, C63, C64, C65, C66, C67, C68, C69, C70, C71, C72, C73, C74, C75, C76, C77, C78, C79,
+		C80, C81, C82, C83, C84, C85, C86, C87, C88, C89, C90, C91, C92, C93, C94, C95, C96, C97, C98, C99);
 }
 
 namespace Expect
 {
+	struct F1_Ctti
+	{
+		static constexpr char Namespace[] = "PlainProps::Test";
+		static constexpr char Name[] = "F1";
+		using Type = ::PlainProps::Test::F1;
+		using enum Type;
+		static constexpr uint32 NumEnumerators = 3;
+		inline static constexpr std::array<std::string_view, NumEnumerators> Names = { "None", "A", "B" };
+		static constexpr Type Constants[NumEnumerators] = { None, A, B };
+	};
+
+	struct U1_Ctti
+	{
+		static constexpr char Namespace[] = "PlainProps::Test";
+		static constexpr char Name[] = "U1";
+		using Type = ::PlainProps::Test::U1;
+		static constexpr uint32 NumEnumerators = 2;
+		inline static constexpr std::array<std::string_view, NumEnumerators> Names = { "U1_A", "U1_B"};
+		static constexpr Type Constants[NumEnumerators] = { U1_A, U1_B };
+	};
+
 	struct E1_Ctti
 	{
+		static constexpr char Namespace[] = "PlainProps::Test";
 		static constexpr char Name[] = "E1";
 		using Type = ::PlainProps::Test::E1;
-		static constexpr int NumEnumerators = 2;
-		static constexpr struct { const char* Name; Type Constant; } Enumerators[] = { {"A", Type::A},  {"B", Type::B} };
+		using enum Type;
+		static constexpr uint32 NumEnumerators = 2;
+		inline static constexpr std::array<std::string_view, NumEnumerators> Names = { "A", "B" };
+		static constexpr Type Constants[NumEnumerators] = { A, B };
+	};
+
+	struct E2_Ctti
+	{
+		static constexpr char Namespace[] = "PlainProps::Test";
+		static constexpr char Name[] = "E2";
+		using Type = ::PlainProps::Test::E2;
+		using enum Type;
+		static constexpr uint32 NumEnumerators = 100;
+		inline static constexpr std::array<std::string_view, NumEnumerators> Names =
+		{
+			"C0","C1","C2","C3","C4","C5","C6","C7","C8","C9",
+			"C10","C11","C12","C13","C14","C15","C16","C17","C18","C19",
+			"C20","C21","C22","C23","C24","C25","C26","C27","C28","C29",
+			"C30","C31","C32","C33","C34","C35","C36","C37","C38","C39",
+			"C40","C41","C42","C43","C44","C45","C46","C47","C48","C49",
+			"C50","C51","C52","C53","C54","C55","C56","C57","C58","C59",
+			"C60","C61","C62","C63","C64","C65","C66","C67","C68","C69",
+			"C70","C71","C72","C73","C74","C75","C76","C77","C78","C79",
+			"C80","C81","C82","C83","C84","C85","C86","C87","C88","C89",
+			"C90","C91","C92","C93","C94","C95","C96","C97","C98","C99"
+		};
+		static constexpr Type Constants[NumEnumerators] =
+		{
+			C0, C1, C2, C3, C4, C5, C6, C7, C8, C9,
+			C10, C11, C12, C13, C14, C15, C16, C17, C18, C19,
+			C20, C21, C22, C23, C24, C25, C26, C27, C28, C29,
+			C30, C31, C32, C33, C34, C35, C36, C37, C38, C39,
+			C40, C41, C42, C43, C44, C45, C46, C47, C48, C49,
+			C50, C51, C52, C53, C54, C55, C56, C57, C58, C59,
+			C60, C61, C62, C63, C64, C65, C66, C67, C68, C69,
+			C70, C71, C72, C73, C74, C75, C76, C77, C78, C79,
+			C80, C81, C82, C83, C84, C85, C86, C87, C88, C89,
+			C90, C91, C92, C93, C94, C95, C96, C97, C98, C99,
+		};
 	};
 }
 
+static_assert(std::string_view(Actual::F1_Ctti::Name) == std::string_view(Expect::F1_Ctti::Name));
+static_assert(std::is_same_v<Actual::F1_Ctti::Type, Expect::F1_Ctti::Type>);
+static_assert(Actual::F1_Ctti::NumEnumerators == Expect::F1_Ctti::NumEnumerators);
+static_assert(Actual::F1_Ctti::Constants[0]	== Expect::F1_Ctti::Constants[0]);
+static_assert(Actual::F1_Ctti::Constants[1]	== Expect::F1_Ctti::Constants[1]);
+static_assert(Actual::F1_Ctti::Constants[2]	== Expect::F1_Ctti::Constants[2]);
+static_assert(std::string_view(Actual::F1_Ctti::Names[0]) == std::string_view(Expect::F1_Ctti::Names[0]));
+static_assert(std::string_view(Actual::F1_Ctti::Names[1]) == std::string_view(Expect::F1_Ctti::Names[1]));
+static_assert(std::string_view(Actual::F1_Ctti::Names[2]) == std::string_view(Expect::F1_Ctti::Names[2]));
+
+static_assert(std::string_view(Actual::U1_Ctti::Name) == std::string_view(Expect::U1_Ctti::Name));
+static_assert(std::is_same_v<Actual::U1_Ctti::Type, Expect::U1_Ctti::Type>);
+static_assert(Actual::U1_Ctti::NumEnumerators == Expect::U1_Ctti::NumEnumerators);
+static_assert(Actual::U1_Ctti::Constants[0]	== Expect::U1_Ctti::Constants[0]);
+static_assert(Actual::U1_Ctti::Constants[1]	== Expect::U1_Ctti::Constants[1]);
+static_assert(std::string_view(Actual::U1_Ctti::Names[0]) == std::string_view(Expect::U1_Ctti::Names[0]));
+static_assert(std::string_view(Actual::U1_Ctti::Names[1]) == std::string_view(Expect::U1_Ctti::Names[1]));
+
 static_assert(std::string_view(Actual::E1_Ctti::Name) == std::string_view(Expect::E1_Ctti::Name));
 static_assert(std::is_same_v<Actual::E1_Ctti::Type, Expect::E1_Ctti::Type>);
 static_assert(Actual::E1_Ctti::NumEnumerators == Expect::E1_Ctti::NumEnumerators);
-static_assert(Actual::E1_Ctti::Enumerators[0].Constant	== Expect::E1_Ctti::Enumerators[0].Constant);
-static_assert(Actual::E1_Ctti::Enumerators[1].Constant	== Expect::E1_Ctti::Enumerators[1].Constant);
-static_assert(std::string_view(Actual::E1_Ctti::Enumerators[0].Name) == std::string_view(Expect::E1_Ctti::Enumerators[0].Name));
-static_assert(std::string_view(Actual::E1_Ctti::Enumerators[1].Name) == std::string_view(Expect::E1_Ctti::Enumerators[1].Name));
+static_assert(Actual::E1_Ctti::Constants[0]	== Expect::E1_Ctti::Constants[0]);
+static_assert(Actual::E1_Ctti::Constants[1]	== Expect::E1_Ctti::Constants[1]);
+static_assert(std::string_view(Actual::E1_Ctti::Names[0]) == std::string_view(Expect::E1_Ctti::Names[0]));
+static_assert(std::string_view(Actual::E1_Ctti::Names[1]) == std::string_view(Expect::E1_Ctti::Names[1]));
+
+static_assert(std::string_view(Actual::E2_Ctti::Name) == std::string_view(Expect::E2_Ctti::Name));
+static_assert(std::is_same_v<Actual::E2_Ctti::Type, Expect::E2_Ctti::Type>);
+static_assert(Actual::E2_Ctti::NumEnumerators == Expect::E2_Ctti::NumEnumerators);
+static_assert(Actual::E2_Ctti::Constants[0]	== Expect::E2_Ctti::Constants[0]);
+static_assert(Actual::E2_Ctti::Constants[1]	== Expect::E2_Ctti::Constants[1]);
+static_assert(Actual::E2_Ctti::Constants[Actual::E2_Ctti::NumEnumerators-2]	== Expect::E2_Ctti::Constants[Expect::E2_Ctti::NumEnumerators-2]);
+static_assert(Actual::E2_Ctti::Constants[Actual::E2_Ctti::NumEnumerators-1]	== Expect::E2_Ctti::Constants[Expect::E2_Ctti::NumEnumerators-1]);
+static_assert(std::string_view(Actual::E2_Ctti::Names[0]) == std::string_view(Expect::E2_Ctti::Names[0]));
+static_assert(std::string_view(Actual::E2_Ctti::Names[1]) == std::string_view(Expect::E2_Ctti::Names[1]));
+static_assert(std::string_view(Actual::E2_Ctti::Names[Actual::E2_Ctti::NumEnumerators-2]) == std::string_view(Expect::E2_Ctti::Names[Expect::E2_Ctti::NumEnumerators-2]));
+static_assert(std::string_view(Actual::E2_Ctti::Names[Actual::E2_Ctti::NumEnumerators-1]) == std::string_view(Expect::E2_Ctti::Names[Expect::E2_Ctti::NumEnumerators-1]));
 
 //////////////////////////////////////////////////////////////////////////
 
@@ -117,4 +228,4 @@ static_assert(CttiOf<S2<int>>::NumVars == 1);
 static_assert(CttiOf<S2<int>>::Var<0>::Name == std::string_view("a"));
 static_assert(CttiOf<S2<int>>::Var<0>::Offset == offsetof(S2<int>, a));
 
-}
\ No newline at end of file
+}
diff --git a/Engine/Plugins/Experimental/PlainProps/Source/Tests/PlainPropsStringTests.cpp b/Engine/Plugins/Experimental/PlainProps/Source/Tests/PlainPropsStringTests.cpp
--- a/Engine/Plugins/Experimental/PlainProps/Source/Tests/PlainPropsStringTests.cpp
+++ b/Engine/Plugins/Experimental/PlainProps/Source/Tests/PlainPropsStringTests.cpp
@@ -6,11 +6,13 @@ namespace PlainProps::Private
 {
 	using namespace std::literals;
 
+	// Concat tests
 	static constexpr std::string_view Hello = "Hello";
 	static constexpr std::string_view Space = " ";
 	static constexpr std::string_view World = "World";
 	static_assert("Hello World"sv == Concat<Hello, Space, World>);
 	
+	// HexString tests
 	static_assert("0"sv == HexString<0>);
 	static_assert("1"sv == HexString<1>);
 	static_assert("9"sv == HexString<9>);
@@ -22,4 +24,23 @@ namespace PlainProps::Private
 	static_assert("FFF"sv == HexString<0xFFF>);
 	static_assert("1000"sv == HexString<0x1000>);
 	static_assert("FEDCBA9876543210"sv == HexString<0xFEDCBA9876543210>);
-}
\ No newline at end of file
+
+	// ParseArgs tests (separated by comma and whitespace)
+	static constexpr std::string_view Many = "A, B,  C,D  ";
+	static constexpr std::string_view One = "\tOne";
+	static constexpr std::string_view Two = "One,\nTwo";
+	static constexpr std::string_view Whitespace = " \v\f\t\r\n";
+	static constexpr std::string_view Empty = "";
+	static constexpr std::array<std::string_view, 4> Args = ParseArgs<Many>;
+	static_assert(Args.size() == 4);
+	static_assert("A"sv == Args[0]);
+	static_assert("B"sv == Args[1]);
+	static_assert("C"sv == Args[2]);
+	static_assert("D"sv == Args[3]);
+	static_assert(ParseArgs<One>.size() == 1);
+	static_assert(ParseArgs<One>[0] == "One"sv);
+	static_assert(ParseArgs<Two>[0] == "One"sv);
+	static_assert(ParseArgs<Two>[1] == "Two"sv);
+	static_assert(ParseArgs<Whitespace>.size() == 0);
+	static_assert(ParseArgs<Empty>.size() == 0);
+}
diff --git a/Engine/Plugins/Experimental/PlainProps/Source/Tests/PlainPropsUeCoreTests.cpp b/Engine/Plugins/Experimental/PlainProps/Source/Tests/PlainPropsUeCoreTests.cpp
--- a/Engine/Plugins/Experimental/PlainProps/Source/Tests/PlainPropsUeCoreTests.cpp
+++ b/Engine/Plugins/Experimental/PlainProps/Source/Tests/PlainPropsUeCoreTests.cpp
@@ -545,7 +545,7 @@ PP_REFLECT_STRUCT(PlainProps::UE::Test, FInt, void, X);
 static bool operator==(FInt A, FInt B) { return A.X == B.X; }
 inline uint32 GetTypeHash(FInt I) { return ::GetTypeHash(I.X); }
 
-enum class EFlat1 : uint8 { A = 1, B = 3 };
+enum class EFlat1 : int8 { A = -1, B = 3 };
 enum class EFlat2 : uint8 { A, B };
 enum class EFlag1 : uint8 { A = 2, B = 8, AB = 10 };
 enum class EFlag2 : uint8 { A = 1, B = 2, AB = 3 };
@@ -565,6 +565,20 @@ struct FEnums
 };
 PP_REFLECT_STRUCT(PlainProps::UE::Test, FEnums, void, Flat1, Flat2, Flag1, Flag2);
 
+enum EByteFlat : int32 { EByteFlat_A = 1, EByteFlat_B = 3 };
+enum EByteFlag : uint8 { EByteFlag_A = 2, EByteFlag_B = 8, EByteFlag_AB = 10 };
+PP_REFLECT_ENUM(PlainProps::UE::Test, EByteFlat, EByteFlat_A, EByteFlat_B);
+PP_REFLECT_ENUM(PlainProps::UE::Test, EByteFlag, EByteFlag_A, EByteFlag_B);
+
+struct FEnumAsBytes
+{
+	TEnumAsByte<EByteFlat> Flat;
+	TEnumAsByte<EByteFlag> Flag;
+
+	friend bool operator==(FEnumAsBytes, FEnumAsBytes) = default;
+};
+PP_REFLECT_STRUCT(PlainProps::UE::Test, FEnumAsBytes, void, Flat, Flag);
+
 struct FLeafArrays
 {
 	TArray<bool> Bits;
@@ -1010,6 +1024,29 @@ TEST_CASE_NAMED(FPlainPropsUeCoreTest, "System::Core::Serialization::PlainProps:
 			});
 	}
 
+	SECTION("TEnumAsByte")
+	{
+		TScopedDefaultEnumDeclaration<EByteFlat, EEnumMode::Flat> Flat;
+		TScopedDefaultEnumDeclaration<EByteFlag, EEnumMode::Flag> Flag;
+		TScopedDefaultStructBinding<FEnumAsBytes> EnumAsBytes;
+		Run([](FBatchSaver& Batch)
+			{
+				Batch.Save(FEnumAsBytes{EByteFlat_A, EByteFlag_A});
+				Batch.Save(FEnumAsBytes{EByteFlat_B, EByteFlag_B});
+				Batch.Save(FEnumAsBytes{EByteFlat_A, EByteFlag_A});
+				Batch.Save(FEnumAsBytes{EByteFlat_B, EByteFlag_B});
+				Batch.Save(FEnumAsBytes{EByteFlat_B, EByteFlag_AB});
+			},
+			[](FBatchLoader& Batch)
+			{
+				CHECK(Batch.Load<FEnumAsBytes>() == FEnumAsBytes{EByteFlat_A, EByteFlag_A});
+				CHECK(Batch.Load<FEnumAsBytes>() == FEnumAsBytes{EByteFlat_B, EByteFlag_B});
+				CHECK(Batch.Load<FEnumAsBytes>() == FEnumAsBytes{EByteFlat_A, EByteFlag_A});
+				CHECK(Batch.Load<FEnumAsBytes>() == FEnumAsBytes{EByteFlat_B, EByteFlag_B});
+				CHECK(Batch.Load<FEnumAsBytes>() == FEnumAsBytes{EByteFlat_B, EByteFlag_AB});
+			});
+	}
+
 	SECTION("TArray")
 	{
 		TScopedDefaultStructBinding<FLeafArrays> LeafArrays;
diff --git a/Engine/Plugins/Experimental/PlainPropsUObject/Source/Private/PlainPropsUObjectBindings.cpp b/Engine/Plugins/Experimental/PlainPropsUObject/Source/Private/PlainPropsUObjectBindings.cpp
--- a/Engine/Plugins/Experimental/PlainPropsUObject/Source/Private/PlainPropsUObjectBindings.cpp
+++ b/Engine/Plugins/Experimental/PlainPropsUObject/Source/Private/PlainPropsUObjectBindings.cpp
@@ -86,6 +86,28 @@ namespace PlainProps::UE
 
 ////////////////////////////////////////////////////////////////////////////////////////////////
 
+inline ELeafWidth WidthOfEnum(const UEnum* Enum)
+{
+	using UnderlyingType = std::underlying_type_t<UEnum::EUnderlyingType>;
+	static_assert(static_cast<UnderlyingType>(UEnum::EUnderlyingType::int8)   == 0);
+	static_assert(static_cast<UnderlyingType>(UEnum::EUnderlyingType::int64)  == 3);
+	static_assert(static_cast<UnderlyingType>(UEnum::EUnderlyingType::uint8)  == 4);
+	static_assert(static_cast<UnderlyingType>(UEnum::EUnderlyingType::uint64) == 7);
+
+	using enum ELeafWidth;
+	static constexpr ELeafWidth Widths[8] = { B8, B16, B32, B64, B8, B16, B32, B64 };
+	UnderlyingType Value = static_cast<UnderlyingType>(Enum->GetUnderlyingType());
+	check(0 <= Value && Value < 8);
+	return Widths[Value];
+}
+
+inline bool IsSigned(const UEnum* Enum)
+{
+	return Enum->GetUnderlyingType() < UEnum::EUnderlyingType::uint8;
+}
+
+////////////////////////////////////////////////////////////////////////////////////////////////
+
 struct FDefaultStruct
 {
 	UScriptStruct::ICppStructOps&		Ops;
@@ -1044,6 +1066,15 @@ static bool DenyMaxValue(FName Enum)
 	return !Algo::Find(AllowsMax, Enum);
 }
 
+static bool FitsInWidth(int64 Value, ELeafWidth Width)
+{
+	const uint32 Bits = 8u << (uint32)Width;
+	const uint32 Shift = 64 - Bits;
+	// The upper bits beyond the width are expected to be all zero (unsigned)
+	// or the Value should survive a truncation + sign-extension roundtrip
+	return (Value >> Bits) == 0 || Value == (static_cast<int64>(static_cast<uint64>(Value) << Shift) >> Shift);
+}
+
 static FEnumId DeclareEnum(UEnum* Enum)
 {
 	FType Type = IndexType(Enum);
@@ -1067,11 +1098,17 @@ static FEnumId DeclareEnum(UEnum* Enum)
 				EndsWithDelimitedSuffix<All>(EnumName, Enum->GetNameByIndex(Num - 1));
 	}
 
+	const ELeafWidth Width = WidthOfEnum(Enum);
+	const uint64 ValueMasks[4] = { uint8(~0), uint16(~0), uint32(~0), ~uint64(0) };
+	const uint64 ValueMask = ValueMasks[(uint8)Width];
 	TArray<FEnumerator, TInlineAllocator<64>> Enumerators;
 	for (int32 Idx = 0; Idx < Num; ++Idx)
 	{
 		FName ValueName = Enum->GetNameByIndex(Idx);
-		Enumerators.Emplace(GUE.Names.MakeName(ValueName), static_cast<uint64>(Enum->GetValueByIndex(Idx)));
+		int64 Value = Enum->GetValueByIndex(Idx);
+		Enumerators.Emplace(GUE.Names.MakeName(ValueName), static_cast<uint64>(Value) & ValueMask);
+		checkf(FitsInWidth(Value, Width), TEXT("Enumerator constant %s=%llx doesn't fit in declared width"),
+			*ValueName.ToString(), Value);
 	}
 
 #if WITH_METADATA
@@ -1079,7 +1116,7 @@ static FEnumId DeclareEnum(UEnum* Enum)
 	checkSlow(Enumerators.Num() == Num || Enum->ContainsExistingMax() || Enum->HasMetaData(TEXT("Hidden"), Num));
 #endif
 
-	GUE.Enums.Declare(Id, Type, Mode, /* size unknown */ {}, Enumerators, EEnumAliases::Strip);
+	GUE.Enums.Declare(Id, Type, Mode, Width, Enumerators, EEnumAliases::Strip);
 	return Id;
 }
 
@@ -4026,9 +4063,8 @@ private:
 		FMemberBinding Out(Property->GetOffset_ForInternal());
 		FEnumId Id = GUE.Names.IndexEnum(IndexType(Property->GetEnum()));
 		Out.InnermostId = FInnerId(Id);
-		FUnpackedLeafType Leaf = {ELeafType::Enum, WidthOf(Property->GetElementSize())};
+		FUnpackedLeafType Leaf = {ELeafType::Enum, WidthOfEnum(Property->GetEnum())};
 		Out.InnermostType = FMemberBindType(Leaf);
-		GUE.Enums.Specify(Id, Leaf.Width); // UEnum doesn't know width/sizeof
 		return Out;
 	}
 
