PlainProps: Minor bugfix in FNameBinding::DiffCustom in PlainPropsUeCoreTests.cpp.

#rnx

#Synced-CL 51785738

diff --git a/Engine/Plugins/Experimental/PlainProps/Source/Tests/PlainPropsUeCoreTests.cpp b/Engine/Plugins/Experimental/PlainProps/Source/Tests/PlainPropsUeCoreTests.cpp
--- a/Engine/Plugins/Experimental/PlainProps/Source/Tests/PlainPropsUeCoreTests.cpp
+++ b/Engine/Plugins/Experimental/PlainProps/Source/Tests/PlainPropsUeCoreTests.cpp
@@ -117,7 +117,7 @@ struct FNameBinding final : ICustomBinding
 	{
 		FName A = *static_cast<const FName*>(StructA);
 		FName B = *static_cast<const FName*>(StructB);
-		return A.IsEqual(B, ENameCase::CaseSensitive);
+		return !A.IsEqual(B, ENameCase::CaseSensitive);
 	}
 };
 
