[Core] Fix cis issue in play props with compact sets enabled

diff --git a/Engine/Plugins/Experimental/PlainProps/Source/Public/PlainPropsUeCoreBindings.h b/Engine/Plugins/Experimental/PlainProps/Source/Public/PlainPropsUeCoreBindings.h
--- a/Engine/Plugins/Experimental/PlainProps/Source/Public/PlainPropsUeCoreBindings.h
+++ b/Engine/Plugins/Experimental/PlainProps/Source/Public/PlainPropsUeCoreBindings.h
@@ -528,8 +528,16 @@ struct TSetBinding : IItemRangeBinding
 #if UE_USE_COMPACT_SET_AS_DEFAULT
 	virtual void ReadItems(FSaveRangeContext& Ctx) const override
 	{
+#if UE_TCOMPACT_SET_ARRAY_VIEW
 		const TArrayView<const T> SetData = Ctx.Request.GetRange<SetType>().ArrayView();
 		Ctx.Items.SetAll(SetData.GetData(), SetData.Num());
+#else
+		// Need to break into jail here, compact sets are linear arrays of data so I can use the address of the first element
+		// to set the rest. ArrayView is currently disabled on it to maintain API compatibility with sparse set.
+		const SetType& Set = Ctx.Request.GetRange<SetType>();
+		const T* SetData = Set.FindArbitraryElement();
+		Ctx.Items.SetAll(SetData, Set.Num());
+#endif
 	}
 #else
 	virtual void ReadItems(FSaveRangeContext& Ctx) const override
