PlainProps: Bug fixes and improved asserts to provide more actionable and human friendly error messages.

Print: Add debug print functions.
Implement ToString() for EMemberKind, FMemberType, FStructType and FRangeType.
Add FDebug::Print(const FMemberSchema&) and fix range bug in PrintMemberSchema(FUtf8Builder& Out, ...).
Files: PlainPropsTypes.h, PlainPropsPrint.h, PlainPropsInternalText.h, PlainPropsPrint.cpp

Build: Detect mismatches in saved vs. declared schema.
Add FMemberSchemaIterator helper.
Add MatchesDeclaredSchema lambda to FMemberBuilder::BuildAndReset().
Decorate assert in FStructSchemaBuilder::NoteMembersRecursively() when detecting a member that has not been declared.
Add early assert about dense substructs to Declare(FStructSpec In), also checked in FMemberBuilder::BuildAndReset().
Improve operator FMemberSpecView() assert by adding a FStructSpec::Check() helper with FDebugIds printing.
Call FStructSpec::Check() from FCustomBindings::BindStruct() to detect when a declared member has not been specified.
Files: PlainPropsSpecify.h, PlainPropsBuild.cpp, PlainPropsBuildSchema.cpp, PlainPropsDeclare.cpp

Specify: Always declare nested ranges in the correct order.
Split FMemberSpec::RangeWrap() into RangeWrap() and RangeAdd().
Files: PlainPropsSpecify.h, PlainPropsBuild.cpp, PlainPropsBindCtti.cpp, PlainPropsDeclare.cpp

Some examples of reported and fixed FMemberBuilder::BuildAndReset() check failures:
Member 'EnumRs' in 'Test.Object' is saved as 'Test.AB(u8)(i32)' but the declared type is 'Test.AB(i32)(u8)'
Member 'Id' in '/Script/CoreUObject.WeakObjectPtr' is saved as 'i32' but the declared schema is 'u64'
Member 'Insert' in '/Script/CoreUObject.Set<[utf8,i32]>' is saved as 'utf8(i32)(i32)' but the declared schema is 'utf8(i32)'

#rnx

#Synced-CL 51891224

diff --git a/Engine/Plugins/Experimental/PlainProps/Source/Private/PlainPropsBind.cpp b/Engine/Plugins/Experimental/PlainProps/Source/Private/PlainPropsBind.cpp
--- a/Engine/Plugins/Experimental/PlainProps/Source/Private/PlainPropsBind.cpp
+++ b/Engine/Plugins/Experimental/PlainProps/Source/Private/PlainPropsBind.cpp
@@ -317,6 +317,20 @@ static bool operator==(FInnerStruct A, FInnerStruct B) { return A.Name == B.Name
 
 ////////////////////////////////////////////////////////////////////////////////////////////////
 
+static void ValidateStructSpec(const FStructSpec& Spec, const FDebugIds& Ids)
+{
+#if DO_CHECK
+	check(Spec.MemberNames.Num() == Spec.MemberTypes.Num());
+	for (int32 Idx = 0, Num = Spec.MemberNames.Num(); Idx < Num; ++Idx)
+	{
+		checkf(Spec.MemberTypes[Idx], TEXT("The declared member '%s' in '%s' has not been specified"),
+			*Ids.Print(Spec.MemberNames[Idx]), *Ids.Print(Spec.Id));
+	}
+#endif
+}
+
+////////////////////////////////////////////////////////////////////////////////////////////////
+
 FCustomBindings::FCustomBindings(TArray<FInnerStruct>& Inners, FDebugIds Dbg) : Map(Dbg), BottomInners(Inners) {}
 FCustomBindings::FCustomBindings(const FCustomBindings* Under) : Map(Under->Map.Debug), BottomInners(Under->BottomInners) {}
 FCustomBindings::~FCustomBindings() {}
@@ -330,6 +344,7 @@ void FCustomBindings::BindStruct(FBindId Id, ICustomBinding& Binding, FStructDec
 
 void FCustomBindings::BindStruct(FBindId Id, ICustomBinding& Binding, const FStructSpec& Spec, TConstArrayView<FInnerStruct> LoweredInners)
 {
+	ValidateStructSpec(Spec, Map.Debug);
 	return BindStruct(Id, Binding, Declare(Spec), LoweredInners);
 }
 
diff --git a/Engine/Plugins/Experimental/PlainProps/Source/Private/PlainPropsBuild.cpp b/Engine/Plugins/Experimental/PlainProps/Source/Private/PlainPropsBuild.cpp
--- a/Engine/Plugins/Experimental/PlainProps/Source/Private/PlainPropsBuild.cpp
+++ b/Engine/Plugins/Experimental/PlainProps/Source/Private/PlainPropsBuild.cpp
@@ -4,6 +4,7 @@
 #include "PlainPropsIndex.h"
 #include "PlainPropsInternalBuild.h"
 #include "PlainPropsInternalFormat.h"
+#include "PlainPropsInternalRead.h"
 #include "Serialization/VarInt.h"
 #include "Templates/TypeCompatibleBytes.h"
 #include "Hash/xxhash.h"
@@ -125,6 +126,55 @@ FBuiltRange* CloneLeaves(FScratchAllocator& Scratch, uint64 Num, const void* InD
 
 //////////////////////////////////////////////////////////////////////////
 
+// Helper for FMemberBuilder::BuildAndReset() validations
+struct FMemberSchemaIterator
+{
+	FMemberSchemaIterator(const FStructDeclaration& Declared)
+		: Members(Declared.GetTypes(), Declared.NumMembers)
+		, InnerRangeTypes(Declared.GetInnerRangeTypes(), Declared.NumInnerRanges)
+		, InnerIds(Declared.GetInnerIds(), Declared.NumInnerIds)
+	{}
+
+	bool HasMore() const
+	{
+		return MemberIdx < Members.Num();
+	}
+
+	FMemberSchema GrabSchema()
+	{
+		check(HasMore());
+		FMemberType Type = Members[MemberIdx++];
+		FMemberType InnerRangeType = Type;
+		FMemberType InnermostType = Type;
+		TConstArrayView<FMemberType> InnerTypes;
+		if (Type.IsRange())
+		{
+			InnerTypes = GrabInnerRangeTypes(InnerRangeTypes, /* in-out */ InnerRangeIdx);
+			InnerRangeType = InnerTypes[0];
+			InnermostType = InnerTypes.Last();
+		}
+		FOptionalInnerId InnerId = HasInnerId(InnermostType) ? ToOptional(InnerIds[InnerIdIdx++]) : NoId;
+		uint16 NumInnerRanges = IntCastChecked<uint16>(InnerTypes.Num());
+		const FMemberType* NestedRangeTypes = NumInnerRanges > 1 ? &InnerTypes[0] : nullptr;
+		return { Type, InnerRangeType, NumInnerRanges, InnerId, NestedRangeTypes };
+	}
+
+private:
+	bool HasInnerId(FMemberType Member) const
+	{
+		return Member.IsStruct() ? !Member.AsStruct().IsDynamic : IsEnum(Member);
+	}
+
+	TConstArrayView<FMemberType> Members;
+	TConstArrayView<FMemberType> InnerRangeTypes;
+	TConstArrayView<FInnerId> InnerIds;
+	uint16 MemberIdx = 0;
+	uint16 InnerRangeIdx = 0;
+	uint16 InnerIdIdx = 0;
+};
+
+//////////////////////////////////////////////////////////////////////////
+
 void FMemberBuilder::AddSuperStruct(FStructId SuperSchema, FBuiltStruct* SuperStruct)
 {
 	check(SuperStruct);
@@ -152,24 +202,55 @@ FBuiltStruct* FMemberBuilder::BuildAndReset(FScratchAllocator& Scratch, const FS
 #if DO_CHECK
 	if (Declared.Occupancy == EMemberPresence::RequireAll)
 	{
-		checkf(Members.Num() == Declared.NumMembers, TEXT("'%s' requires exactly %d members but %d were added"), *Debug.Print(Declared.Id), Declared.NumMembers, Members.Num());
-		checkf(!Declared.Super, TEXT("Bug, dense substructs should fail DeclareStruct() check"));
+		checkf(!Declared.Super, TEXT("Bug, dense substructs should fail DeclareStruct() check."));
+		checkf(Members.Num() == Declared.NumMembers, TEXT("'%s' requires exactly %d members but %d were added."), *Debug.Print(Declared.Id), Declared.NumMembers, Members.Num());
 	}
 	
-	// Verify members were added in declared order
+	auto MatchesDeclaredSchema = [](FMemberSchema Built, FMemberSchema Declared) -> bool
+	{
+		// For structs, match anything since declared struct schema may be dynamic or type-erased
+		if (Built.Type.IsStruct() && Declared.Type.IsStruct())
+		{
+			return true;
+		}
+		// For struct ranges, match the range size and nested range sizes, but skip the innermost type and schema
+		else if (Built.Type.IsRange() &&
+			Built.GetInnerRangeTypes().Last().IsStruct() && Declared.GetInnerRangeTypes().Last().IsStruct())
+		{
+			return Built.Type == Declared.Type &&
+				Algo::Compare(Built.GetInnerRangeTypes().LeftChop(1), Declared.GetInnerRangeTypes().LeftChop(1));
+		}
+		// Match all fields
+		return Built.Type == Declared.Type &&
+			Built.InnerSchema == Declared.InnerSchema &&
+			Algo::Compare(Built.GetInnerRangeTypes(), Declared.GetInnerRangeTypes());
+	};
+
+	// Verify members were added in declared order and that they match the declared schemas.
 	if (int32 Num = Members.Num())
 	{
 		const bool HasSuper = IsSuper(Members[0].Schema.Type);
 		int32 OrderIdx = 0;
 		TConstArrayView<FMemberId> Order = Declared.GetMemberOrder();
 		int32 SkipSuper = Declared.Super && HasSuper;
+		FMemberSchemaIterator DeclaredIt(Declared);
 		for (FBuiltMember* It = Members.GetData() + SkipSuper, *End = Members.GetData() + Num; It != End; ++It)
 		{
 			FBuiltMember& Member = *It;
 			for (; OrderIdx < Order.Num() && Order[OrderIdx] != Member.Name; ++OrderIdx)
-			{}	
-			checkf(OrderIdx < Order.Num(), TEXT("Member '%s' in '%s' %s"), *Debug.Print(Member.Name), *Debug.Print(Declared.Id),
-					Order.Contains(Member.Name) ? TEXT("appeared in non-declared order") : TEXT("is undeclared"));
+			{
+				(void)DeclaredIt.GrabSchema();
+			}
+			checkf(OrderIdx < Order.Num(),
+				TEXT("Member '%s' in '%s' %s."),
+				*Debug.Print(Member.Name), *Debug.Print(Declared.Id),
+				Order.Contains(Member.Name) ? TEXT("appeared in non-declared order") : TEXT("is undeclared"));
+
+			const FMemberSchema DeclaredSchema = DeclaredIt.GrabSchema();
+			checkf(MatchesDeclaredSchema(Member.Schema, DeclaredSchema),
+				TEXT("Member '%s' in '%s' is saved as '%s' but the declared schema is '%s'."),
+				*Debug.Print(Member.Name), *Debug.Print(Declared.Id),
+				*Debug.Print(Member.Schema), *Debug.Print(DeclaredSchema));
 			++OrderIdx;
 		}
 	}
diff --git a/Engine/Plugins/Experimental/PlainProps/Source/Private/PlainPropsBuildSchema.cpp b/Engine/Plugins/Experimental/PlainProps/Source/Private/PlainPropsBuildSchema.cpp
--- a/Engine/Plugins/Experimental/PlainProps/Source/Private/PlainPropsBuildSchema.cpp
+++ b/Engine/Plugins/Experimental/PlainProps/Source/Private/PlainPropsBuildSchema.cpp
@@ -252,11 +252,13 @@ void FStructSchemaBuilder::NoteMembersRecursively(const FBuiltStruct& Struct)
 	int32 NoteIdx = 0;
 	for (const FBuiltMember& Member : MakeArrayView(Struct.Members, Struct.NumMembers))
 	{
-		while (MemberOrder[NoteIdx] != Member.Name)
+		while (NoteIdx < NumNoted && MemberOrder[NoteIdx] != Member.Name)
 		{
 			++NoteIdx;
-			check(NoteIdx < NumNoted);
 		}
+		checkf(NoteIdx < NumNoted, TEXT("Member '%s' has not been declared in '%s'"),
+			*AllSchemas.GetDebug().Print(Member.Name),
+			*AllSchemas.GetDebug().Print(Declaration.Id));
 
 		if (NotedMembers[NoteIdx])
 		{
diff --git a/Engine/Plugins/Experimental/PlainProps/Source/Private/PlainPropsDeclare.cpp b/Engine/Plugins/Experimental/PlainProps/Source/Private/PlainPropsDeclare.cpp
--- a/Engine/Plugins/Experimental/PlainProps/Source/Private/PlainPropsDeclare.cpp
+++ b/Engine/Plugins/Experimental/PlainProps/Source/Private/PlainPropsDeclare.cpp
@@ -68,6 +68,7 @@ static TConstArrayView<FEnumerator> StripAliases(TArray<FEnumerator, TInlineAllo
 
 static void ValidateDeclaration(const FEnumDeclaration& Enum)
 {
+#if DO_CHECK
 	if (Enum.Mode == EEnumMode::Flag)
 	{
 		for (FEnumerator E : Enum.GetEnumerators())
@@ -93,6 +94,7 @@ static void ValidateDeclaration(const FEnumDeclaration& Enum)
 	{
 		checkf(!(E.Constant & OverflowMask), TEXT("Enumerator constant %llx is larger than declared width"), E.Constant);
 	}
+#endif
 }
 
 template<typename T>
@@ -161,6 +163,7 @@ bool FStructDeclaration::Release() const
 
 FStructDeclarationPtr Declare(FStructSpec In) 
 {
+	checkf(!(In.Super && In.Occupancy == EMemberPresence::RequireAll), TEXT("Dense substructs are not supported"));
 	check(In.MemberNames.Num() == In.MemberTypes.Num());
 
 	// Extract inner ids
@@ -221,7 +224,7 @@ FMemberSpec::FMemberSpec(ESizeType MaxSize, FMemberSpec Inner)
 FMemberSpec::FMemberSpec(TConstArrayView<FMemberType> Members, FOptionalInnerId InInnermostId)
 : FMemberSpec(Members.Last(), InInnermostId)
 {
-	for (FMemberType Range : Members.LeftChop(1))
+	for (FMemberType Range : ReverseIterate(Members.LeftChop(1)))
 	{
 		RangeWrap(Range);
 	}
@@ -232,11 +235,11 @@ FMemberSpec::FMemberSpec(FMemberType Type, TConstArrayView<FMemberType> InnerRan
 {
 	if (int32 NumRanges = InnerRangeTypes.Num())
 	{
-		RangeWrap(Type);
-		for (FMemberType Nested : InnerRangeTypes.Slice(0, NumRanges - 1))
+		for (FMemberType Nested : ReverseIterate(InnerRangeTypes.Slice(0, NumRanges - 1)))
 		{
 			RangeWrap(Nested);	
 		}
+		RangeWrap(Type);
 	}
 }
 
diff --git a/Engine/Plugins/Experimental/PlainProps/Source/Private/PlainPropsInternalText.h b/Engine/Plugins/Experimental/PlainProps/Source/Private/PlainPropsInternalText.h
--- a/Engine/Plugins/Experimental/PlainProps/Source/Private/PlainPropsInternalText.h
+++ b/Engine/Plugins/Experimental/PlainProps/Source/Private/PlainPropsInternalText.h
@@ -33,9 +33,14 @@ struct FLiterals
 
 	FAnsiStringView Oob = "ERR!oob";
 
+	FAnsiStringView Kinds[3] = { ("Leaf"), ("Struct"), ("Range") };
+
 	FAnsiStringView Ranges[9] = { 
 		("Uni"), ("i8"), ("u8"), ("i16"), ("u16"), ("i32"), ("u32"), ("i64"), ("u64") };
 
+	FAnsiStringView RangeTypes[9] = { 
+		("(Uni)"), ("(i8)"), ("(u8)"), ("(i16)"), ("(u16)"), ("(i32)"), ("(u32)"), ("(i64)"), ("(u64)") };
+
 	FAnsiStringView Leaves[8][4] = {
 		{("bool"),		("ERR!b16"),	("ERR!b32"),	("ERR!b64")},
 		{("i8"),		("i16"),		("i32"),		("i64")},
diff --git a/Engine/Plugins/Experimental/PlainProps/Source/Private/PlainPropsPrint.cpp b/Engine/Plugins/Experimental/PlainProps/Source/Private/PlainPropsPrint.cpp
--- a/Engine/Plugins/Experimental/PlainProps/Source/Private/PlainPropsPrint.cpp
+++ b/Engine/Plugins/Experimental/PlainProps/Source/Private/PlainPropsPrint.cpp
@@ -28,6 +28,16 @@ FAnsiStringView ToString(ESizeType Width)
 	return GLiterals.Ranges[(uint8)Width];
 }
 
+FAnsiStringView ToString(FRangeType Range)
+{
+	return GLiterals.RangeTypes[(uint8)Range.MaxSize];
+}
+
+FAnsiStringView ToString(EMemberKind Kind)
+{
+	return GLiterals.Kinds[(uint8)Kind];
+}
+
 FAnsiStringView ToString(FUnpackedLeafType Leaf)
 {
 	return GLiterals.Leaves[(uint8)Leaf.Type][(uint8)Leaf.Width];
@@ -38,6 +48,35 @@ FAnsiStringView ToString(ELeafWidth Width)
 	return GLiterals.Widths[(uint8)Width];
 }
 
+FAnsiStringView ToString(FStructType Struct)
+{
+	if (Struct.IsSuper && Struct.IsDynamic)
+	{
+		return "DynamicSuperStruct";
+	}
+	else if (Struct.IsSuper)
+	{
+		return "SuperStruct";
+	}
+	else if (Struct.IsDynamic)
+	{
+		return "DynamicStruct";
+	}
+	return "Struct";
+}
+
+FAnsiStringView ToString(FMemberType Type)
+{
+	switch (Type.GetKind())
+	{
+	case EMemberKind::Leaf:		return ToString(Type.AsLeaf());
+	case EMemberKind::Struct:	return ToString(Type.AsStruct());
+	case EMemberKind::Range:	return ToString(Type.AsRange());
+	}
+	unimplemented();
+	return "Unknown";
+}
+
 FAnsiStringView ToString(ESchemaFormat Format)
 {
 	switch (Format)
@@ -283,7 +322,8 @@ void PrintMemberSchema(FUtf8Builder& Out, const FIds& Ids, FMemberSchema Member)
 	if (Member.Type.IsRange())
 	{
 		PrintRangeType(Out, Member.Type.AsRange());
-		for (FMemberType Inner : Member.GetInnerRangeTypes())
+		check(!Member.GetInnerRangeTypes().Last().IsRange());
+		for (FMemberType Inner : Member.GetInnerRangeTypes().LeftChop(1))
 		{
 			PrintRangeType(Out, Inner.AsRange());
 		}
@@ -1830,6 +1870,13 @@ FString FDebugIds::Print(FStructId Name) const
 	return FString(FStringView(StringCast<TCHAR>(Out.GetData(), Out.Len())));
 }
 
+FString FDebugIds::Print(const FMemberSchema& Schema) const
+{
+	TUtf8StringBuilder<128> Out;
+	PrintMemberSchema(Out, Ids, Schema);
+	return FString(FStringView(StringCast<TCHAR>(Out.GetData(), Out.Len())));
+}
+
 ///////////////////////////////////////////////////////////////////////////////
 
 void PrintDiff(FUtf8Builder& Out, const FIds& Ids, const FDiffPath& Diff)
diff --git a/Engine/Plugins/Experimental/PlainProps/Source/Public/PlainPropsBindCtti.h b/Engine/Plugins/Experimental/PlainProps/Source/Public/PlainPropsBindCtti.h
--- a/Engine/Plugins/Experimental/PlainProps/Source/Public/PlainPropsBindCtti.h
+++ b/Engine/Plugins/Experimental/PlainProps/Source/Public/PlainPropsBindCtti.h
@@ -532,7 +532,7 @@ FMemberBinding BindMember(uint64 Offset, FMemberSpec& OutSpec)
 			Out.RangeBindings = GetRangeBindings<RangeBinding, Ids, NumRangeBindings>();
 			Out.InnermostType = BindInnermostType<InnermostType, Ids>(Out.InnermostId, OutSpec);
 			
-			for (FRangeBinding Range : ReverseIterate(Out.RangeBindings))
+			for (FRangeBinding Range : Out.RangeBindings)
 			{
 				OutSpec.RangeWrap(Range.GetSizeType());
 			}
diff --git a/Engine/Plugins/Experimental/PlainProps/Source/Public/PlainPropsPrint.h b/Engine/Plugins/Experimental/PlainProps/Source/Public/PlainPropsPrint.h
--- a/Engine/Plugins/Experimental/PlainProps/Source/Public/PlainPropsPrint.h
+++ b/Engine/Plugins/Experimental/PlainProps/Source/Public/PlainPropsPrint.h
@@ -9,16 +9,20 @@ namespace PlainProps
 {
 
 class FBatchIds;
+class FSchemaBatchId;
 struct FDiffPath;
 struct FMemberSchema;
-class FSchemaBatchId;
 struct FReadDiffPath;
 struct FStructView;
 
 PLAINPROPS_API FAnsiStringView ToString(ELeafWidth Width);
+PLAINPROPS_API FAnsiStringView ToString(EMemberKind Kind);
 PLAINPROPS_API FAnsiStringView ToString(ESizeType Width);
 PLAINPROPS_API FAnsiStringView ToString(ESchemaFormat Format);
+PLAINPROPS_API FAnsiStringView ToString(FMemberType Member);
+PLAINPROPS_API FAnsiStringView ToString(FStructType Struct);
 PLAINPROPS_API FAnsiStringView ToString(FUnpackedLeafType Leaf);
+PLAINPROPS_API FAnsiStringView ToString(FRangeType Range);
 
 ///////////////////////////////////////////////////////////////////////////////
 
diff --git a/Engine/Plugins/Experimental/PlainProps/Source/Public/PlainPropsSpecify.h b/Engine/Plugins/Experimental/PlainProps/Source/Public/PlainPropsSpecify.h
--- a/Engine/Plugins/Experimental/PlainProps/Source/Public/PlainPropsSpecify.h
+++ b/Engine/Plugins/Experimental/PlainProps/Source/Public/PlainPropsSpecify.h
@@ -49,10 +49,11 @@ public:
 	PLAINPROPS_API FMemberSpec(TConstArrayView<FMemberType> Members, FOptionalInnerId InInnermostId);
 	PLAINPROPS_API FMemberSpec(FMemberType Type, TConstArrayView<FMemberType> InnerRangeTypes, FOptionalInnerId InInnermostId);
 	
-	void RangeWrap(ESizeType MaxSize)		{ Ranges.Emplace(MaxSize); }
-	void RangeWrap(FMemberType Range)		{ Ranges.Emplace(Range.AsRange().MaxSize); }
+	void RangeWrap(ESizeType MaxSize)		{ Ranges.Insert(MaxSize, 0); }
+	void RangeWrap(FMemberType Range)		{ Ranges.Insert(Range.AsRange().MaxSize, 0); }
 
 	operator FMemberSpecView() const UE_LIFETIMEBOUND;
+	explicit operator bool() const { return bSpecified; }
 	
 private:
 	using RangeSizeArray = TArray<ESizeType, TInlineAllocator<8>>;
@@ -67,7 +68,7 @@ private:
 
 inline FMemberSpec::operator FMemberSpecView() const
 {
-	check(bSpecified);
+	checkf(bSpecified, TEXT("Bug, unspecified members should fail earlier checks."));
 	return { InnermostType, InnermostId, Ranges };
 }
 
@@ -120,4 +121,4 @@ struct FSoleMemberSpec
 };
 
 
-} // namespace PlainProps
\ No newline at end of file
+} // namespace PlainProps
diff --git a/Engine/Plugins/Experimental/PlainProps/Source/Public/PlainPropsTypes.h b/Engine/Plugins/Experimental/PlainProps/Source/Public/PlainPropsTypes.h
--- a/Engine/Plugins/Experimental/PlainProps/Source/Public/PlainPropsTypes.h
+++ b/Engine/Plugins/Experimental/PlainProps/Source/Public/PlainPropsTypes.h
@@ -557,8 +557,10 @@ public:
 	PLAINPROPS_API FString			Print(FType Type) const;
 	PLAINPROPS_API FString			Print(FEnumId Enum) const;
 	PLAINPROPS_API FString			Print(FStructId Struct) const;
+	PLAINPROPS_API FString			Print(const struct FMemberSchema& Schema) const;
+
 private:
 	const FIds& Ids;
 };
 
-} // namespace PlainProps
\ No newline at end of file
+} // namespace PlainProps
diff --git a/Engine/Plugins/Experimental/PlainProps/Source/Tests/PlainPropsTests.cpp b/Engine/Plugins/Experimental/PlainProps/Source/Tests/PlainPropsTests.cpp
--- a/Engine/Plugins/Experimental/PlainProps/Source/Tests/PlainPropsTests.cpp
+++ b/Engine/Plugins/Experimental/PlainProps/Source/Tests/PlainPropsTests.cpp
@@ -1289,7 +1289,7 @@ TEST_CASE_NAMED(FPlainPropsReadWriteTest, "System::Core::Serialization::PlainPro
 			FEnumId Enum = Batch.DeclareEnum("Test", "AB", EEnumMode::Flag, ELeafWidth::B8, {"A", "B"}, {1, 4});
 			FEnumId UnusedEnum = Batch.DeclareEnum("Test", "Unused", EEnumMode::Flat, ELeafWidth::B8, {"X"}, {0});
 			static constexpr std::initializer_list<const char*> MemberNames = { "IntRs", "EmptyRs", "FloatRs", "EnumRs", "UnusedEnumRs", "UnicodeRs", "StructRs", "StructRRs" };
-			FMemberSpec MemberTypes[] = { S32Range(S32Range(SpecS32)), S32Range(S32Range(SpecBool)), S64Range(S32Range(SpecF32)), 
+			FMemberSpec MemberTypes[] = { S32Range(S32Range(SpecS32)), S32Range(S32Range(SpecS32)), S64Range(S64Range(SpecF32)),
 				U8Range(S32Range(Enum8(Enum))), U8Range(S32Range(Enum8(UnusedEnum))), S32Range(S32Range(SpecUtf8)),
 				U64Range(U64Range(XY)), U32Range(S16Range(S16Range(ZW)))};
 			FDeclId Object = Batch.DeclareStruct("Test", "Object", MemberNames, MemberTypes, EMemberPresence::AllowSparse);
@@ -1332,8 +1332,8 @@ TEST_CASE_NAMED(FPlainPropsReadWriteTest, "System::Core::Serialization::PlainPro
 			FMemberSchema ZWRangeSchema = MakeStructRangeSchema(ESizeType::S16, ZW);
 			FNestedRangeBuilder ZWRs(ZWRangeSchema, 1);
 			ZWRs.Add(ZWs.BuildAndReset(Scratch, Batch.Get(ZW), Batch.GetDebug()));
-			FNestedRangeBuilder StructRRs(MakeNestedRangeSchema(Scratch, ESizeType::U32, ZWRangeSchema), 1);
-			StructRRs.Add(ZWRs.BuildAndReset(Scratch, ESizeType::U32));
+			FNestedRangeBuilder StructRRs(MakeNestedRangeSchema(Scratch, ESizeType::S16, ZWRangeSchema), 1);
+			StructRRs.Add(ZWRs.BuildAndReset(Scratch, ESizeType::S16));
 
 			FMemberBuilder Members;
 			Members.AddRange(Batch.NameMember("IntRs"), IntRs.BuildAndReset(Scratch, ESizeType::S32));
@@ -1411,7 +1411,7 @@ TEST_CASE_NAMED(FPlainPropsReadWriteTest, "System::Core::Serialization::PlainPro
 	{
 		TestSerialize([](FTestBatchBuilder& Batch, FScratchAllocator& Scratch)
 		{
-			FMemberSpec StructMemberTypes[] = { UniRange(SpecBool), S32Range(SpecBool), UniRange(SpecBool), SpecBool };
+			FMemberSpec StructMemberTypes[] = { UniRange(SpecBool), S32Range(SpecBool), S32Range(UniRange(SpecBool)), SpecBool };
 			FDeclId Struct = Batch.DeclareStruct("Test", "Struct", {"MaybeB", "Bs", "MaybeBs", "B"}, StructMemberTypes, EMemberPresence::AllowSparse);
 			FMemberSpec ObjectMemberTypes[] = { UniRange(SpecBool), S32Range(Struct), SpecBool, SpecBool };
 			FDeclId Object = Batch.DeclareStruct("Test", "Object", {"Bools", "Structs", "BF", "BT" }, ObjectMemberTypes, EMemberPresence::AllowSparse);
@@ -1425,10 +1425,10 @@ TEST_CASE_NAMED(FPlainPropsReadWriteTest, "System::Core::Serialization::PlainPro
 			Structs[7].AddRange(Batch.NameMember("MaybeB"), BuildLeafRange(Scratch, &True, true));
 			Structs[7].AddRange(Batch.NameMember("Bs"),		BuildLeafRange(Scratch, MakeArrayView({true, true, false, false, true, true, false, false, true, true})));
 			MaybeBs.Add(BuildLeafRange(Scratch, &True, true));
-			Structs[7].AddRange(Batch.NameMember("MaybeBs"), MaybeBs.BuildAndReset(Scratch, ESizeType::Uni));
+			Structs[7].AddRange(Batch.NameMember("MaybeBs"), MaybeBs.BuildAndReset(Scratch, ESizeType::S32));
 			Structs[7].Add(Batch.NameMember("B"), true);
 			MaybeBs.Add(BuildLeafRange(Scratch, &True, false));
-			Structs[8].AddRange(Batch.NameMember("MaybeBs"), MaybeBs.BuildAndReset(Scratch, ESizeType::Uni));
+			Structs[8].AddRange(Batch.NameMember("MaybeBs"), MaybeBs.BuildAndReset(Scratch, ESizeType::S32));
 			Structs[9].Add(Batch.NameMember("B"), false);
 
 			FMemberBuilder Members;
@@ -1520,9 +1520,9 @@ TEST_CASE_NAMED(FPlainPropsReadWriteTest, "System::Core::Serialization::PlainPro
 	{
 		TestSerialize([](FTestBatchBuilder& Batch, FScratchAllocator& Scratch)
 		{
-			FDeclId SA = Batch.DeclareStruct("Test", "SA", {"X"}, {S32Range(SpecS32)}, EMemberPresence::AllowSparse);
-			FDeclId Unused = Batch.DeclareStruct("Test", "Unused2", {"X"}, {S32Range(SpecS32)}, EMemberPresence::AllowSparse);
-			FDeclId SB = Batch.DeclareStruct("Test", "SB", {"X"}, {S32Range(SpecF32)}, EMemberPresence::AllowSparse);
+			FDeclId SA = Batch.DeclareStruct("Test", "SA", {"X"}, {SpecS32}, EMemberPresence::AllowSparse);
+			FDeclId Unused = Batch.DeclareStruct("Test", "Unused2", {"X"}, {SpecS32}, EMemberPresence::AllowSparse);
+			FDeclId SB = Batch.DeclareStruct("Test", "SB", {"X"}, {SpecF32}, EMemberPresence::AllowSparse);
 			FMemberSpec MemberTypes[] = { S32Range(SA), S32Range(SA), SpecBool, S32Range(SpecDynamicStruct), S32Range(SA), S32Range(SpecDynamicStruct), S32Range(S32Range(SpecDynamicStruct)) };
 			FDeclId Object = Batch.DeclareStruct("Test", "Object", {"Same", "Some", "None", "Diff", "SameEmpty", "DiffEmpty", "DiffNested"}, MemberTypes, EMemberPresence::AllowSparse);
 		
diff --git a/Engine/Plugins/Experimental/PlainPropsUObject/Source/Private/PlainPropsRoundtripTest.cpp b/Engine/Plugins/Experimental/PlainPropsUObject/Source/Private/PlainPropsRoundtripTest.cpp
--- a/Engine/Plugins/Experimental/PlainPropsUObject/Source/Private/PlainPropsRoundtripTest.cpp
+++ b/Engine/Plugins/Experimental/PlainPropsUObject/Source/Private/PlainPropsRoundtripTest.cpp
@@ -1233,12 +1233,11 @@ struct FLinkerPropertyBatch
 
 	static FSoleMemberSpec Spec(FWeakObjectPtr*, FDeclId Id)
 	{
-		return {Id, GUE.Members.Id, Specify<uint64>()};
+		return {Id, GUE.Members.Id, Specify<int32>()};
 	}
 
 	void Save(FMemberBuilder& Out, const FWeakObjectPtr& In, const FSaveContext& Ctx)
 	{
-		// Save ObjectSerialNumber + ObjectIndex as a single uint64
 		UObject* Object = In.Get(true);
 		Save(Out, ::UE::CoreUObject::Private::MakeObjectHandle(Object), Ctx);
 	}
diff --git a/Engine/Plugins/Experimental/PlainPropsUObject/Source/Private/PlainPropsUObjectBindings.cpp b/Engine/Plugins/Experimental/PlainPropsUObject/Source/Private/PlainPropsUObjectBindings.cpp
--- a/Engine/Plugins/Experimental/PlainPropsUObject/Source/Private/PlainPropsUObjectBindings.cpp
+++ b/Engine/Plugins/Experimental/PlainPropsUObject/Source/Private/PlainPropsUObjectBindings.cpp
@@ -1848,7 +1848,7 @@ struct FNestedRangeSerializer
 	inline FMemberSpec SpecMember() const
 	{
 		// InnermostSaveId is safe since FProperty doesnt support type-erasure
-		return FMemberSpec(InnerTypes, InnermostSaveId);
+		return DefaultRangeOf(FMemberSpec(InnerTypes, InnermostSaveId));
 	}
 
 	FBuiltRange* SaveItem(const void* In, const FSaveContext& Ctx) const
@@ -3610,7 +3610,7 @@ static FMemberSpec ToSpec(FMemberBinding Binding)
 		? FMemberSpec(ToLeafType(Binding.InnermostType.AsLeaf()), ToOptionalEnum(Binding.InnermostId))
 		: FMemberSpec(Binding.InnermostId.Get().AsStructDeclId());
 
-	for (FRangeBinding RangeBinding : Binding.RangeBindings)
+	for (FRangeBinding RangeBinding : ReverseIterate(Binding.RangeBindings))
 	{
 		Out.RangeWrap(RangeBinding.GetSizeType());
 	}
