PlainProps: Fix compile warning for PlainPropsDeclare.cpp.

#rnx

#Synced-CL 51891224

diff --git a/Engine/Plugins/Experimental/PlainProps/Source/Private/PlainPropsDeclare.cpp b/Engine/Plugins/Experimental/PlainProps/Source/Private/PlainPropsDeclare.cpp
--- a/Engine/Plugins/Experimental/PlainProps/Source/Private/PlainPropsDeclare.cpp
+++ b/Engine/Plugins/Experimental/PlainProps/Source/Private/PlainPropsDeclare.cpp
@@ -224,7 +224,8 @@ FMemberSpec::FMemberSpec(ESizeType MaxSize, FMemberSpec Inner)
 FMemberSpec::FMemberSpec(TConstArrayView<FMemberType> Members, FOptionalInnerId InInnermostId)
 : FMemberSpec(Members.Last(), InInnermostId)
 {
-	for (FMemberType Range : ReverseIterate(Members.LeftChop(1)))
+	TConstArrayView<FMemberType> RangeTypes = Members.LeftChop(1);
+	for (FMemberType Range : ReverseIterate(RangeTypes))
 	{
 		RangeWrap(Range);
 	}
@@ -235,7 +236,8 @@ FMemberSpec::FMemberSpec(FMemberType Type, TConstArrayView<FMemberType> InnerRan
 {
 	if (int32 NumRanges = InnerRangeTypes.Num())
 	{
-		for (FMemberType Nested : ReverseIterate(InnerRangeTypes.Slice(0, NumRanges - 1)))
+		TConstArrayView<FMemberType> RangeTypes = InnerRangeTypes.Slice(0, NumRanges - 1);
+		for (FMemberType Nested : ReverseIterate(RangeTypes))
 		{
 			RangeWrap(Nested);	
 		}
