PlainProps: Add missing TTypename<FText>::DeclName.

#rnx

#Synced-CL 52055016

diff --git a/Engine/Plugins/Experimental/PlainProps/Source/Public/PlainPropsUeCoreBindings.h b/Engine/Plugins/Experimental/PlainProps/Source/Public/PlainPropsUeCoreBindings.h
--- a/Engine/Plugins/Experimental/PlainProps/Source/Public/PlainPropsUeCoreBindings.h
+++ b/Engine/Plugins/Experimental/PlainProps/Source/Public/PlainPropsUeCoreBindings.h
@@ -1204,6 +1204,7 @@ namespace PlainProps
 {
 
 template <> struct TTypename<FName>					{ inline static constexpr std::string_view DeclName = "Name"; };
+template <> struct TTypename<FText>					{ inline static constexpr std::string_view DeclName = "Text"; };
 template <> struct TTypename<FIntPoint>				{ inline static constexpr std::string_view DeclName = "IntPoint"; };
 template <> struct TTypename<FQuat>					{ inline static constexpr std::string_view DeclName = "Quat"; };
 template <> struct TTypename<FTransform>			{ inline static constexpr std::string_view DeclName = "Transform"; };
