PlainProps: Decorate assert for non-indexed members in custom bindings by prinding the BindId name.
Add explicit TScopedStructBinding(FDualStructId Dual) constructor for easier construction of schema based bindings.

#rnx

#Synced-CL 52055016

diff --git a/Engine/Plugins/Experimental/PlainProps/Source/Public/PlainPropsBindCtti.h b/Engine/Plugins/Experimental/PlainProps/Source/Public/PlainPropsBindCtti.h
--- a/Engine/Plugins/Experimental/PlainProps/Source/Public/PlainPropsBindCtti.h
+++ b/Engine/Plugins/Experimental/PlainProps/Source/Public/PlainPropsBindCtti.h
@@ -314,7 +314,9 @@ struct TScopedStructBinding : FBothStructId, CustomBinding
 	: FBothStructId(Both)
 	, CustomBinding(/* out */ Init)
 	{
-		checkf(N > 0 && CustomBinding::MemberIds[N - 1] != FMemberId{}, TEXT("MemberIds not fully initialized"));
+		checkf(N > 0 && CustomBinding::MemberIds[N - 1] != FMemberId{},
+			TEXT("Invalid members in custom binding for '%s' (Num=%u)"),
+			*FDebugIds(Ids::GetIndexer()).Print(Both.BindId), N);
 		FStructSpec Decl = { DeclId, Init.Super, Init.Version, TOccupancyOf<T>::Value, CustomBinding::MemberIds, Init.Members };
 		Runtime::GetCustoms().BindStruct(BindId, *this, Decl, Init.LoweredInners);
 	}
@@ -367,7 +369,11 @@ struct TScopedStructBinding<T, Runtime, void> : FBothStructId
 {
 	using Ids = typename Runtime::Ids;
 	using Typename = TTypename<T>;
-	
+
+	explicit TScopedStructBinding(FDualStructId Dual)
+	: TScopedStructBinding(FBothStructId{Dual, Dual})
+	{}
+
 	TScopedStructBinding(FBothStructId Both = IndexStructBothId<Ids, Typename>() )
 	: FBothStructId(Both)
 	{
