PlainProps: Extend PP_REFLECT_STRUCT macro and ForEachVar template to 16 members (required for blueprint bindings) and add a static_assert to catch overflow in ForEachVar.

#rnx

#Synced-CL 52129858

diff --git a/Engine/Plugins/Experimental/PlainProps/Source/Public/PlainPropsCtti.h b/Engine/Plugins/Experimental/PlainProps/Source/Public/PlainPropsCtti.h
--- a/Engine/Plugins/Experimental/PlainProps/Source/Public/PlainPropsCtti.h
+++ b/Engine/Plugins/Experimental/PlainProps/Source/Public/PlainPropsCtti.h
@@ -35,8 +35,9 @@ constexpr bool ValidateEnumFlags()
 // Private PP_REFLECT_XYZ helpers
 
 #define _PP_EXPAND(X) X
-#define _PP_NUM_ARGS_(_1,_2,_3,_4,_5,_6,_7,_8,_9,_10, N, ...) N
-#define _PP_NUM_ARGS(...) _PP_EXPAND(_PP_NUM_ARGS_(__VA_ARGS__,10,9,8,7,6,5,4,3,2,1))
+
+#define _PP_NUM_ARGS_(_1,_2,_3,_4,_5,_6,_7,_8,_9,_10,_11,_12,_13,_14,_15,_16, N, ...) N
+#define _PP_NUM_ARGS(...) _PP_NUM_ARGS_(__VA_ARGS__,16,15,14,13,12,11,10,9,8,7,6,5,4,3,2,1)
 
 #define _PP_REFLECT_1( N, MACRO, NS, T, M)				MACRO(N-1, NS, T, M)
 #define _PP_REFLECT_2( N, MACRO, NS, T, M, ...)			MACRO(N-2, NS, T, M)  _PP_REFLECT_1(N, MACRO, NS, T, __VA_ARGS__)
@@ -48,6 +49,12 @@ constexpr bool ValidateEnumFlags()
 #define _PP_REFLECT_8( N, MACRO, NS, T, M, ...)			MACRO(N-8, NS, T, M)  _PP_REFLECT_7(N, MACRO, NS, T, __VA_ARGS__)
 #define _PP_REFLECT_9( N, MACRO, NS, T, M, ...)			MACRO(N-9, NS, T, M)  _PP_REFLECT_8(N, MACRO, NS, T, __VA_ARGS__)
 #define _PP_REFLECT_10(N, MACRO, NS, T, M, ...)			MACRO(N-10, NS, T, M) _PP_REFLECT_9(N, MACRO, NS, T, __VA_ARGS__)
+#define _PP_REFLECT_11(N, MACRO, NS, T, M, ...)			MACRO(N-11, NS, T, M) _PP_REFLECT_10(N, MACRO, NS, T, __VA_ARGS__)
+#define _PP_REFLECT_12(N, MACRO, NS, T, M, ...)			MACRO(N-12, NS, T, M) _PP_REFLECT_11(N, MACRO, NS, T, __VA_ARGS__)
+#define _PP_REFLECT_13(N, MACRO, NS, T, M, ...)			MACRO(N-13, NS, T, M) _PP_REFLECT_12(N, MACRO, NS, T, __VA_ARGS__)
+#define _PP_REFLECT_14(N, MACRO, NS, T, M, ...)			MACRO(N-14, NS, T, M) _PP_REFLECT_13(N, MACRO, NS, T, __VA_ARGS__)
+#define _PP_REFLECT_15(N, MACRO, NS, T, M, ...)			MACRO(N-15, NS, T, M) _PP_REFLECT_14(N, MACRO, NS, T, __VA_ARGS__)
+#define _PP_REFLECT_16(N, MACRO, NS, T, M, ...)			MACRO(N-16, NS, T, M) _PP_REFLECT_15(N, MACRO, NS, T, __VA_ARGS__)
 #define _PP_REFLECT_THINGS(N, MACRO, ...)				_PP_EXPAND(_PP_REFLECT_##N(N, MACRO, __VA_ARGS__))
 
 #define _PP_REFLECT_ENUM_INNER(NS, T, S, ...)			struct T##_Ctti \
@@ -92,6 +99,7 @@ template <class Ctti, typename Fn>
 static constexpr void ForEachVar(Fn&& Visitor)
 {
     constexpr int N = Ctti::NumVars;
+    static_assert(N <= 16, "Member overflow");
     if constexpr (N > 0) Visitor.template operator()<typename Ctti::template Var<0>>();
     if constexpr (N > 1) Visitor.template operator()<typename Ctti::template Var<1>>();
     if constexpr (N > 2) Visitor.template operator()<typename Ctti::template Var<2>>();
@@ -102,6 +110,12 @@ static constexpr void ForEachVar(Fn&& Visitor)
     if constexpr (N > 7) Visitor.template operator()<typename Ctti::template Var<7>>();
     if constexpr (N > 8) Visitor.template operator()<typename Ctti::template Var<8>>();
     if constexpr (N > 9) Visitor.template operator()<typename Ctti::template Var<9>>();
+    if constexpr (N > 10) Visitor.template operator()<typename Ctti::template Var<10>>();
+    if constexpr (N > 11) Visitor.template operator()<typename Ctti::template Var<11>>();
+    if constexpr (N > 12) Visitor.template operator()<typename Ctti::template Var<12>>();
+    if constexpr (N > 13) Visitor.template operator()<typename Ctti::template Var<13>>();
+    if constexpr (N > 14) Visitor.template operator()<typename Ctti::template Var<14>>();
+    if constexpr (N > 15) Visitor.template operator()<typename Ctti::template Var<15>>();
 }
 
 //////////////////////////////////////////////////////////////////////////
