PlainProps: Add index to ValidateStructSpec assert.

#rnx

#Synced-CL 52129858

diff --git a/Engine/Plugins/Experimental/PlainProps/Source/Private/PlainPropsBind.cpp b/Engine/Plugins/Experimental/PlainProps/Source/Private/PlainPropsBind.cpp
--- a/Engine/Plugins/Experimental/PlainProps/Source/Private/PlainPropsBind.cpp
+++ b/Engine/Plugins/Experimental/PlainProps/Source/Private/PlainPropsBind.cpp
@@ -323,8 +323,8 @@ static void ValidateStructSpec(const FStructSpec& Spec, const FDebugIds& Ids)
 	check(Spec.MemberNames.Num() == Spec.MemberTypes.Num());
 	for (int32 Idx = 0, Num = Spec.MemberNames.Num(); Idx < Num; ++Idx)
 	{
-		checkf(Spec.MemberTypes[Idx], TEXT("The declared member '%s' in '%s' has not been specified"),
-			*Ids.Print(Spec.MemberNames[Idx]), *Ids.Print(Spec.Id));
+		checkf(Spec.MemberTypes[Idx], TEXT("The declared member '%s' at member index '%d' in '%s' has not been specified"),
+			*Ids.Print(Spec.MemberNames[Idx]), Idx, *Ids.Print(Spec.Id));
 	}
 #endif
 }
