PlainProps: Extend EnumAsByte test to test saving and loading both E and EnumAsByte<E>.

#rnx

#Synced-CL 52377540

diff --git a/Engine/Plugins/Experimental/PlainProps/Source/Tests/PlainPropsUeCoreTests.cpp b/Engine/Plugins/Experimental/PlainProps/Source/Tests/PlainPropsUeCoreTests.cpp
--- a/Engine/Plugins/Experimental/PlainProps/Source/Tests/PlainPropsUeCoreTests.cpp
+++ b/Engine/Plugins/Experimental/PlainProps/Source/Tests/PlainPropsUeCoreTests.cpp
@@ -572,12 +572,14 @@ PP_REFLECT_ENUM(PlainProps::UE::Test, EByteFlag, EByteFlag_A, EByteFlag_B);
 
 struct FEnumAsBytes
 {
-	TEnumAsByte<EByteFlat> Flat;
-	TEnumAsByte<EByteFlag> Flag;
+	EByteFlat Flat;
+	EByteFlag Flag;
+	TEnumAsByte<EByteFlat> FlatAsByte;
+	TEnumAsByte<EByteFlag> FlagAsByte;
 
 	friend bool operator==(FEnumAsBytes, FEnumAsBytes) = default;
 };
-PP_REFLECT_STRUCT(PlainProps::UE::Test, FEnumAsBytes, void, Flat, Flag);
+PP_REFLECT_STRUCT(PlainProps::UE::Test, FEnumAsBytes, void, Flat, Flag, FlatAsByte, FlagAsByte);
 
 struct FLeafArrays
 {
@@ -1031,19 +1033,19 @@ TEST_CASE_NAMED(FPlainPropsUeCoreTest, "System::Core::Serialization::PlainProps:
 		TScopedDefaultStructBinding<FEnumAsBytes> EnumAsBytes;
 		Run([](FBatchSaver& Batch)
 			{
-				Batch.Save(FEnumAsBytes{EByteFlat_A, EByteFlag_A});
-				Batch.Save(FEnumAsBytes{EByteFlat_B, EByteFlag_B});
-				Batch.Save(FEnumAsBytes{EByteFlat_A, EByteFlag_A});
-				Batch.Save(FEnumAsBytes{EByteFlat_B, EByteFlag_B});
-				Batch.Save(FEnumAsBytes{EByteFlat_B, EByteFlag_AB});
+				Batch.Save(FEnumAsBytes{EByteFlat_A, EByteFlag_A,  EByteFlat_A,  EByteFlag_A});
+				Batch.Save(FEnumAsBytes{EByteFlat_B, EByteFlag_B,  EByteFlat_B,  EByteFlag_B});
+				Batch.Save(FEnumAsBytes{EByteFlat_A, EByteFlag_A,  EByteFlat_A,  EByteFlag_A});
+				Batch.Save(FEnumAsBytes{EByteFlat_B, EByteFlag_B,  EByteFlat_B,  EByteFlag_B});
+				Batch.Save(FEnumAsBytes{EByteFlat_B, EByteFlag_AB, EByteFlat_B,  EByteFlag_AB});
 			},
 			[](FBatchLoader& Batch)
 			{
-				CHECK(Batch.Load<FEnumAsBytes>() == FEnumAsBytes{EByteFlat_A, EByteFlag_A});
-				CHECK(Batch.Load<FEnumAsBytes>() == FEnumAsBytes{EByteFlat_B, EByteFlag_B});
-				CHECK(Batch.Load<FEnumAsBytes>() == FEnumAsBytes{EByteFlat_A, EByteFlag_A});
-				CHECK(Batch.Load<FEnumAsBytes>() == FEnumAsBytes{EByteFlat_B, EByteFlag_B});
-				CHECK(Batch.Load<FEnumAsBytes>() == FEnumAsBytes{EByteFlat_B, EByteFlag_AB});
+				CHECK(Batch.Load<FEnumAsBytes>() == FEnumAsBytes{EByteFlat_A, EByteFlag_A,  EByteFlat_A, EByteFlag_A});
+				CHECK(Batch.Load<FEnumAsBytes>() == FEnumAsBytes{EByteFlat_B, EByteFlag_B,  EByteFlat_B, EByteFlag_B});
+				CHECK(Batch.Load<FEnumAsBytes>() == FEnumAsBytes{EByteFlat_A, EByteFlag_A,  EByteFlat_A, EByteFlag_A});
+				CHECK(Batch.Load<FEnumAsBytes>() == FEnumAsBytes{EByteFlat_B, EByteFlag_B,  EByteFlat_B, EByteFlag_B});
+				CHECK(Batch.Load<FEnumAsBytes>() == FEnumAsBytes{EByteFlat_B, EByteFlag_AB, EByteFlat_B, EByteFlag_AB});
 			});
 	}
 
