PlainProps: Decorate assert in FEnumSchemaBuilder::NoteValue with more information.
Example: Enum 'DesiredPinDirection' in '/Script/BlueprintGraph.UserPinInfo' first added as 'u32' and later as 'u8' for value '1'

#rnx

#Synced-CL 52377540

diff --git a/Engine/Plugins/Experimental/PlainProps/Source/Private/PlainPropsBuildSchema.cpp b/Engine/Plugins/Experimental/PlainProps/Source/Private/PlainPropsBuildSchema.cpp
--- a/Engine/Plugins/Experimental/PlainProps/Source/Private/PlainPropsBuildSchema.cpp
+++ b/Engine/Plugins/Experimental/PlainProps/Source/Private/PlainPropsBuildSchema.cpp
@@ -433,7 +433,13 @@ FBuiltEnumSchema FEnumSchemaBuilder::Build() const
 
 void FEnumSchemaBuilder::NoteValue(ELeafWidth Width, uint64 Value, FStructId Struct, FMemberId Member)
 {
-	check(NotedWidth == NullOpt || NotedWidth == Width);
+	checkf(NotedWidth == NullOpt || NotedWidth == Width,
+		TEXT("Enum '%s' in '%s' first added as '%.*hs' and later as '%.*hs' for value '%llx'."),
+		*Debug.Print(Member), *Debug.Print(Struct),
+		ToString(*NotedWidth).Len(), ToString(*NotedWidth).GetData(),
+		ToString(Width).Len(), ToString(Width).GetData(),
+		Value);
+
 	NotedWidth = Width;
 
 	if (Declaration.Mode == EEnumMode::Flag)
