PlainProps: Split TEnumAsByteBinding into TAuxEnumAsByteBinding and TUEnumAsByteBinding to align with the expected UProperty behavior.
Save UEnum more efficiently as enum8, and load it more efficiently with a memcpy plan.

#rnx

#Synced-CL 52377540

diff --git a/Engine/Plugins/Experimental/PlainProps/Source/Public/PlainPropsUeCoreBindings.h b/Engine/Plugins/Experimental/PlainProps/Source/Public/PlainPropsUeCoreBindings.h
--- a/Engine/Plugins/Experimental/PlainProps/Source/Public/PlainPropsUeCoreBindings.h
+++ b/Engine/Plugins/Experimental/PlainProps/Source/Public/PlainPropsUeCoreBindings.h
@@ -643,7 +643,7 @@ struct TMapBinding : public TSetBinding<TPair<K, V>, KeyFuncs, SetAllocator>
 //////////////////////////////////////////////////////////////////////////
 
 template <Enumeration E>
-struct TEnumAsByteBinding : ICustomBinding
+struct TAuxEnumAsByteBinding : ICustomBinding
 {
 	using Type = TEnumAsByte<E>;
 
@@ -651,29 +651,25 @@ struct TEnumAsByteBinding : ICustomBinding
 	const FEnumId				EnumId;
 
 	template<class Ids>
-	TEnumAsByteBinding(TCustomSpecifier<Ids, 1>& Spec)
+	TAuxEnumAsByteBinding(TCustomSpecifier<Ids, 1>& Spec)
 	: MemberIds{Ids::IndexMember("Value")}
 	, EnumId(GetEnumId<Ids, E>())
 	{
 		Spec.Members[0] = Specify<E>(EnumId);
 	}
 
-	inline void Save(FMemberBuilder& Dst, const TEnumAsByte<E>& Src, const TEnumAsByte<E>* Default, const FSaveContext& Ctx) const
+	void Save(FMemberBuilder& Dst, const TEnumAsByte<E>& Src, const TEnumAsByte<E>* Default, const FSaveContext& Ctx) const
 	{
 		Dst.AddEnum(MemberIds[0], EnumId, Src.GetValue());
 	}
 		
-	inline void Load(TEnumAsByte<E>& Dst, FStructLoadView Src, ECustomLoadMethod Method) const
+	void Load(TEnumAsByte<E>& Dst, FStructLoadView Src, ECustomLoadMethod Method) const
 	{
 		FMemberLoader Members(Src);
 		Dst = Members.GrabLeaf().As<E>();
 	}
 
-	template<typename ContextType>
-	bool Diff(TEnumAsByte<E> A, TEnumAsByte<E> B, ContextType& Ctx) const
-	{
-		return A != B;
-	}
+	static bool Diff(TEnumAsByte<E> A, TEnumAsByte<E> B, const FBindContext& Ctx) { return A != B; }
 };
 
 //////////////////////////////////////////////////////////////////////////
@@ -1356,9 +1352,10 @@ struct TCustomBind<TPair<K, V>>
 template<Enumeration E>
 struct TOccupancyOf<TEnumAsByte<E>> : FRequireAll {};
 
-template<Enumeration E> struct TCustomBind<TEnumAsByte<E>>
+template<Enumeration E> requires (requires { CttiOfPtr((E*)nullptr); })
+struct TCustomBind<TEnumAsByte<E>>
 {
-	using Type = UE::TEnumAsByteBinding<E>;
+	using Type = UE::TAuxEnumAsByteBinding<E>;
 };
 
 template<typename... Ts>
diff --git a/Engine/Plugins/Experimental/PlainPropsUObject/Source/Public/PlainPropsUObjectRuntime.h b/Engine/Plugins/Experimental/PlainPropsUObject/Source/Public/PlainPropsUObjectRuntime.h
--- a/Engine/Plugins/Experimental/PlainPropsUObject/Source/Public/PlainPropsUObjectRuntime.h
+++ b/Engine/Plugins/Experimental/PlainPropsUObject/Source/Public/PlainPropsUObjectRuntime.h
@@ -234,6 +234,32 @@ using TScopedDefaultUClassBinding = TScopedUClassBinding<T, FDefaultRuntime, Cus
 
 //////////////////////////////////////////////////////////////////////////
 
+template <Enumeration E>
+struct TUEnumAsByteBinding : ICustomBinding
+{
+	using Type = TEnumAsByte<E>;
+
+	const FMemberId				MemberIds[1];
+	const FEnumId				EnumId;
+
+	template<class Ids>
+	TUEnumAsByteBinding(TPropertySpecifier<1>& Spec)
+	: MemberIds{GUE.Members.Value}
+	, EnumId(GetEnumId<Ids, E>())
+	{
+		static_assert(sizeof(Type) == 1);
+		Spec.Members[0] = FMemberSpec(ELeafType::Enum, ELeafWidth::B8, EnumId);
+	}
+
+	void Save(FMemberBuilder& Dst, const TEnumAsByte<E>& Src, const TEnumAsByte<E>* Default, const FSaveContext& Ctx) const
+	{
+		Dst.AddEnum(MemberIds[0], EnumId, Src.GetIntValue());
+	}
+	static void Plan(FMemcpyLoadPlan& Out) { Out.Size = sizeof(Type); }
+	static bool Diff(TEnumAsByte<E> A, TEnumAsByte<E> B, const FBindContext&) { return A != B; }
+};
+
+//////////////////////////////////////////////////////////////////////////
 struct FFieldPathBinding : ICustomBinding
 {
 	using Type = FFieldPath;
@@ -312,6 +338,9 @@ struct FReferencePropertyBinding : ICustomBinding
 
 namespace PlainProps
 {
+template<Enumeration E> requires (requires { StaticEnum<E>(); })
+struct TCustomBind<TEnumAsByte<E>> { using Type = UE::TUEnumAsByteBinding<E>; };
+
 template<> struct TCustomBind<FFieldPath> { using Type = UE::FFieldPathBinding; };
 template<> struct TCustomBind<FScriptDelegate> { using Type = UE::FDelegateBinding; };
 template<> struct TCustomBind<FMulticastScriptDelegate> { using Type = UE::FMulticastInlineDelegateBinding; };
