PlainProps: TExternallyBound struct trait to skip BindCustomStructOnce singleton for externally bound structs (ex: FGuid, FColor, FLinearColor, FTransform)

#rnx
#Synced-CL 52539327

diff --git a/Engine/Plugins/Experimental/PlainProps/Source/Public/PlainPropsBindCtti.h b/Engine/Plugins/Experimental/PlainProps/Source/Public/PlainPropsBindCtti.h
--- a/Engine/Plugins/Experimental/PlainProps/Source/Public/PlainPropsBindCtti.h
+++ b/Engine/Plugins/Experimental/PlainProps/Source/Public/PlainPropsBindCtti.h
@@ -26,6 +26,9 @@ using RangeBind = typename TRangeBind<T>::Type;
 template<typename T>
 struct TCustomBind { using Type = void; };
 
+template<typename T, class Runtime>
+struct TExternallyBound : std::false_type {};
+
 template<class T>
 struct TCustomDeltaBind { using Type = void; };
 
@@ -524,7 +527,17 @@ FMemberBinding BindMember(uint64 Offset, FMemberSpec& OutSpec)
 	FMemberBinding Out(Offset);
 	if constexpr (!std::is_void_v<CustomBinding>)
 	{
-		FBothStructId Both = BindCustomStructOnce<Type, CustomBinding, Runtime>();
+		FBothStructId Both;  
+		if constexpr (TExternallyBound<Type, Runtime>::value)
+		{
+			using Typename = typename TCustomTypename<CustomBinding>::Type;
+			Both = IndexStructBothId<Ids, Typename>();
+			check(Runtime::GetCustoms().FindStruct(Both.BindId) != nullptr);
+		}
+		else
+		{
+			Both = BindCustomStructOnce<Type, CustomBinding, Runtime>();
+		}
 		Out.InnermostType = BindInnerStruct(Out.InnermostId, OutSpec, Both);
 	}
 	else
diff --git a/Engine/Plugins/Experimental/PlainPropsUObject/Source/Public/PlainPropsUObjectRuntime.h b/Engine/Plugins/Experimental/PlainPropsUObject/Source/Public/PlainPropsUObjectRuntime.h
--- a/Engine/Plugins/Experimental/PlainPropsUObject/Source/Public/PlainPropsUObjectRuntime.h
+++ b/Engine/Plugins/Experimental/PlainPropsUObject/Source/Public/PlainPropsUObjectRuntime.h
@@ -407,6 +407,11 @@ namespace PlainProps
 template<Enumeration E> requires (requires { StaticEnum<E>(); })
 struct TCustomBind<TEnumAsByte<E>> { using Type = UE::TUEnumAsByteBinding<E>; };
 
+template<> struct TExternallyBound<FGuid, UE::FDefaultRuntime> : std::true_type {};
+template<> struct TExternallyBound<FColor, UE::FDefaultRuntime> : std::true_type {};
+template<> struct TExternallyBound<FLinearColor, UE::FDefaultRuntime> : std::true_type {};
+template<> struct TExternallyBound<FTransform, UE::FDefaultRuntime> : std::true_type {};
+
 // Enable usage of TObjectPtr<T> as members in ctti reflected structs
 template<typename T> struct	TTypename<TObjectPtr<T>>	{ static_assert(!sizeof(T), "Unsupported type for TObjectPtr"); };
 template<> struct TTypename<TObjectPtr<UClass>>			{ inline static constexpr std::string_view DeclName = "ClassPtr"; };
